Fixed bad asset pointer cache logic

#rnx

diff --git a/Engine/Plugins/Editor/ContentBrowser/ContentBrowserAssetDataSource/Source/ContentBrowserAssetDataSource/Private/ContentBrowserAssetDataPayload.cpp b/Engine/Plugins/Editor/ContentBrowser/ContentBrowserAssetDataSource/Source/ContentBrowserAssetDataSource/Private/ContentBrowserAssetDataPayload.cpp
--- a/Engine/Plugins/Editor/ContentBrowser/ContentBrowserAssetDataSource/Source/ContentBrowserAssetDataSource/Private/ContentBrowserAssetDataPayload.cpp
+++ b/Engine/Plugins/Editor/ContentBrowser/ContentBrowserAssetDataSource/Source/ContentBrowserAssetDataSource/Private/ContentBrowserAssetDataPayload.cpp
@@ -61,7 +61,7 @@ UPackage* FContentBrowserAssetFileItemDataPayload::LoadPackage() const
 
 UObject* FContentBrowserAssetFileItemDataPayload::GetAsset(const bool bTryRecacheIfNull) const
 {
-	if (!bHasCachedAssetPtr || (bTryRecacheIfNull && !CachedPackagePtr.IsValid()))
+	if (!bHasCachedAssetPtr || (bTryRecacheIfNull && !CachedAssetPtr.IsValid()))
 	{
 		if (!AssetData.ObjectPath.IsNone())
 		{
@@ -76,7 +76,7 @@ UObject* FContentBrowserAssetFileItemDataPayload::GetAsset(const bool bTryRecach
 
 UObject* FContentBrowserAssetFileItemDataPayload::LoadAsset() const
 {
-	if (!bHasCachedAssetPtr || !CachedPackagePtr.IsValid())
+	if (!bHasCachedAssetPtr || !CachedAssetPtr.IsValid())
 	{
 		if (!AssetData.ObjectPath.IsNone())
 		{
