Fixed a path too long warning

#rnx

diff --git a/Engine/Plugins/Editor/ContentBrowser/ContentBrowserAssetDataSource/Source/ContentBrowserAssetDataSource/ContentBrowserAssetDataSource.Build.cs b/Engine/Plugins/Editor/ContentBrowser/ContentBrowserAssetDataSource/Source/ContentBrowserAssetDataSource/ContentBrowserAssetDataSource.Build.cs
--- a/Engine/Plugins/Editor/ContentBrowser/ContentBrowserAssetDataSource/Source/ContentBrowserAssetDataSource/ContentBrowserAssetDataSource.Build.cs
+++ b/Engine/Plugins/Editor/ContentBrowser/ContentBrowserAssetDataSource/Source/ContentBrowserAssetDataSource/ContentBrowserAssetDataSource.Build.cs
@@ -6,6 +6,8 @@ namespace UnrealBuildTool.Rules
 	{
 		public ContentBrowserAssetDataSource(ReadOnlyTargetRules Target) : base(Target)
 		{
+			ShortName = "CBAssetDataSource";
+
 			PublicDependencyModuleNames.AddRange(
 				new string[] {
 					"Core",
diff --git a/Engine/Plugins/Editor/ContentBrowser/ContentBrowserClassDataSource/Source/ContentBrowserClassDataSource/ContentBrowserClassDataSource.Build.cs b/Engine/Plugins/Editor/ContentBrowser/ContentBrowserClassDataSource/Source/ContentBrowserClassDataSource/ContentBrowserClassDataSource.Build.cs
--- a/Engine/Plugins/Editor/ContentBrowser/ContentBrowserClassDataSource/Source/ContentBrowserClassDataSource/ContentBrowserClassDataSource.Build.cs
+++ b/Engine/Plugins/Editor/ContentBrowser/ContentBrowserClassDataSource/Source/ContentBrowserClassDataSource/ContentBrowserClassDataSource.Build.cs
@@ -6,6 +6,8 @@ namespace UnrealBuildTool.Rules
 	{
 		public ContentBrowserClassDataSource(ReadOnlyTargetRules Target) : base(Target)
 		{
+			ShortName = "CBClassDataSource";
+
 			PublicDependencyModuleNames.AddRange(
 				new string[] {
 					"Core",
diff --git a/Engine/Plugins/Editor/ContentBrowser/ContentBrowserFileDataSource/Source/ContentBrowserFileDataSource/ContentBrowserFileDataSource.Build.cs b/Engine/Plugins/Editor/ContentBrowser/ContentBrowserFileDataSource/Source/ContentBrowserFileDataSource/ContentBrowserFileDataSource.Build.cs
--- a/Engine/Plugins/Editor/ContentBrowser/ContentBrowserFileDataSource/Source/ContentBrowserFileDataSource/ContentBrowserFileDataSource.Build.cs
+++ b/Engine/Plugins/Editor/ContentBrowser/ContentBrowserFileDataSource/Source/ContentBrowserFileDataSource/ContentBrowserFileDataSource.Build.cs
@@ -6,6 +6,8 @@ namespace UnrealBuildTool.Rules
 	{
 		public ContentBrowserFileDataSource(ReadOnlyTargetRules Target) : base(Target)
 		{
+			ShortName = "CBFileDataSource";
+
 			PublicDependencyModuleNames.AddRange(
 				new string[] {
 					"Core",
