UContentBrowserAssetDataSource::OnAssetLoaded now ignores PIE Packages.

diff --git a/Engine/Plugins/Editor/ContentBrowser/ContentBrowserAssetDataSource/Source/ContentBrowserAssetDataSource/Private/ContentBrowserAssetDataSource.cpp b/Engine/Plugins/Editor/ContentBrowser/ContentBrowserAssetDataSource/Source/ContentBrowserAssetDataSource/Private/ContentBrowserAssetDataSource.cpp
--- a/Engine/Plugins/Editor/ContentBrowser/ContentBrowserAssetDataSource/Source/ContentBrowserAssetDataSource/Private/ContentBrowserAssetDataSource.cpp
+++ b/Engine/Plugins/Editor/ContentBrowser/ContentBrowserAssetDataSource/Source/ContentBrowserAssetDataSource/Private/ContentBrowserAssetDataSource.cpp
@@ -1468,7 +1468,7 @@ void UContentBrowserAssetDataSource::OnAssetUpdated(const FAssetData& InAssetDat
 
 void UContentBrowserAssetDataSource::OnAssetLoaded(UObject* InAsset)
 {
-	if (InAsset && !InAsset->GetOutermost()->HasAnyPackageFlags(PKG_ForDiffing))
+	if (InAsset && !InAsset->GetOutermost()->HasAnyPackageFlags(PKG_ForDiffing | PKG_PlayInEditor))
 	{
 		FAssetData AssetData(InAsset);
 		if (ContentBrowserAssetData::IsPrimaryAsset(AssetData))
