Added package disk size as an hardcoded asset data item attribute so it display as a tooltip or column in the content browser

diff --git a/Engine/Plugins/Editor/AssetManagerEditor/Source/AssetManagerEditor/AssetManagerEditor.Build.cs b/Engine/Plugins/Editor/AssetManagerEditor/Source/AssetManagerEditor/AssetManagerEditor.Build.cs
--- a/Engine/Plugins/Editor/AssetManagerEditor/Source/AssetManagerEditor/AssetManagerEditor.Build.cs
+++ b/Engine/Plugins/Editor/AssetManagerEditor/Source/AssetManagerEditor/AssetManagerEditor.Build.cs
@@ -33,6 +33,7 @@ public class AssetManagerEditor : ModuleRules
 				"Json",
 				"CollectionManager",
 				"ContentBrowser",
+				"ContentBrowserData",
 				"WorkspaceMenuStructure",
 				"EditorStyle",
 				"AssetTools",
diff --git a/Engine/Plugins/Editor/AssetManagerEditor/Source/AssetManagerEditor/Private/SAssetAuditBrowser.cpp b/Engine/Plugins/Editor/AssetManagerEditor/Source/AssetManagerEditor/Private/SAssetAuditBrowser.cpp
--- a/Engine/Plugins/Editor/AssetManagerEditor/Source/AssetManagerEditor/Private/SAssetAuditBrowser.cpp
+++ b/Engine/Plugins/Editor/AssetManagerEditor/Source/AssetManagerEditor/Private/SAssetAuditBrowser.cpp
@@ -20,6 +20,7 @@
 #include "ClassViewerFilter.h"
 #include "IContentBrowserSingleton.h"
 #include "ContentBrowserModule.h"
+#include "ContentBrowserDataSource.h"
 #include "AssetRegistryModule.h"
 #include "Toolkits/GlobalEditorCommonCommands.h"
 #include "FrontendFilterBase.h"
@@ -295,6 +296,8 @@ void SAssetAuditBrowser::Construct(const FArguments& InArgs)
 	// Hide path and type by default
 	Config.HiddenColumnNames.Add(TEXT("Class"));
 	Config.HiddenColumnNames.Add(TEXT("Path"));
+	// Hide item disk size since we display a custom version here
+	Config.HiddenColumnNames.Add(ContentBrowserItemAttributes::ItemDiskSize.ToString());
 
 	// Add custom columns
 	Config.CustomColumns.Emplace(FPrimaryAssetId::PrimaryAssetTypeTag, LOCTEXT("AssetType", "Primary Type"), LOCTEXT("AssetTypeTooltip", "Primary Asset Type of this asset, if set"), UObject::FAssetRegistryTag::TT_Alphabetical, FOnGetCustomAssetColumnData::CreateSP(this, &SAssetAuditBrowser::GetStringValueForCustomColumn), FOnGetCustomAssetColumnDisplayText::CreateSP(this, &SAssetAuditBrowser::GetDisplayTextForCustomColumn));
diff --git a/Engine/Plugins/Editor/ContentBrowser/ContentBrowserAssetDataSource/Source/ContentBrowserAssetDataSource/Private/ContentBrowserAssetDataCore.cpp b/Engine/Plugins/Editor/ContentBrowser/ContentBrowserAssetDataSource/Source/ContentBrowserAssetDataSource/Private/ContentBrowserAssetDataCore.cpp
--- a/Engine/Plugins/Editor/ContentBrowser/ContentBrowserAssetDataSource/Source/ContentBrowserAssetDataSource/Private/ContentBrowserAssetDataCore.cpp
+++ b/Engine/Plugins/Editor/ContentBrowser/ContentBrowserAssetDataSource/Source/ContentBrowserAssetDataSource/Private/ContentBrowserAssetDataCore.cpp
@@ -1072,6 +1072,30 @@ void GetClassItemAttribute(const FAssetData& InAssetData, const bool InIncludeMe
 	}
 }
 
+bool GetDiskSizeItemAttribute(const FAssetData& InAssetData, IAssetRegistry* InAssetRegistry, const bool InIncludeMetaData, FContentBrowserItemDataAttributeValue& OutAttributeValue)
+{
+	check(InAssetData.IsValid());
+	check(InAssetRegistry);
+
+	if (const FAssetPackageData* PackageData = InAssetRegistry->GetAssetPackageData(InAssetData.PackageName))
+	{
+		OutAttributeValue.SetValue(PackageData->DiskSize);
+
+		if (InIncludeMetaData)
+		{
+			static const FText DiskSizeDisplayName = LOCTEXT("AttributeDisplayName_DiskSize", "Disk Size");
+
+			FContentBrowserItemDataAttributeMetaData AttributeMetaData;
+			AttributeMetaData.AttributeType = UObject::FAssetRegistryTag::TT_Numerical;
+			AttributeMetaData.DisplayFlags = UObject::FAssetRegistryTag::TD_Memory;
+			AttributeMetaData.DisplayName = DiskSizeDisplayName;
+			OutAttributeValue.SetMetaData(MoveTemp(AttributeMetaData));
+		}
+		return true;
+	}
+	return false;
+}
+
 void GetGenericItemAttribute(const FName InTagKey, const FString& InTagValue, const FAssetPropertyTagCache::FClassPropertyTagCache& InClassPropertyTagCache, const bool InIncludeMetaData, FContentBrowserItemDataAttributeValue& OutAttributeValue)
 {
 	check(!InTagKey.IsNone());
@@ -1203,6 +1227,11 @@ bool GetAssetFileItemAttribute(const FContentBrowserAssetFileItemDataPayload& In
 			return false;
 		}
 
+		if (InAttributeKey == ContentBrowserItemAttributes::ItemDiskSize)
+		{
+			return GetDiskSizeItemAttribute(InAssetPayload.GetAssetData(), IAssetRegistry::Get(), InIncludeMetaData, OutAttributeValue);
+		}
+
 		if (InAttributeKey == ContentBrowserItemAttributes::ItemIsDeveloperContent)
 		{
 			const bool bIsDevelopersFolder = AssetViewUtils::IsDevelopersFolder(InAssetPayload.GetAssetData().PackageName.ToString());
@@ -1290,8 +1319,16 @@ bool GetAssetFileItemAttributes(const FContentBrowserAssetFileItemDataPayload& I
 {
 	// Hard-coded attribute keys
 	{
+		// Class
 		FContentBrowserItemDataAttributeValue& ClassAttributeValue = OutAttributeValues.Add(NAME_Class);
 		GetClassItemAttribute(InAssetPayload.GetAssetData(), InIncludeMetaData, ClassAttributeValue);
+
+		// Disk Size
+		FContentBrowserItemDataAttributeValue DiskSizeAttributeValue;
+		if (GetDiskSizeItemAttribute(InAssetPayload.GetAssetData(), IAssetRegistry::Get(), InIncludeMetaData, DiskSizeAttributeValue))
+		{
+			OutAttributeValues.Add(ContentBrowserItemAttributes::ItemDiskSize, MoveTemp(DiskSizeAttributeValue));
+		}
 	}
 
 	// Generic attribute keys
diff --git a/Engine/Source/Editor/ContentBrowserData/Public/ContentBrowserDataSource.h b/Engine/Source/Editor/ContentBrowserData/Public/ContentBrowserDataSource.h
--- a/Engine/Source/Editor/ContentBrowserData/Public/ContentBrowserDataSource.h
+++ b/Engine/Source/Editor/ContentBrowserData/Public/ContentBrowserDataSource.h
@@ -35,6 +35,12 @@ namespace ContentBrowserItemAttributes
 	 */
 	const FName ItemDescription = "ItemDescription";
 
+	/**
+	 * Attribute key that can be used to query the internal disk size of an item.
+	 * Type: int64.
+	 */
+	const FName ItemDiskSize = "ItemDiskSize";
+
 	/**
 	 * Attribute key that can be used to query whether the given item is considered to be developer content.
 	 * Type: bool.
