Remove the ContentBrowserAssetDataSource's OnAssetLoaded hook. The OnAssetAdded and OnAssetUpdated hooks from the AssetRegistry already provide the same events, and the OnAssetLoaded was duplicating the (sometimes costly) effort of collecting the AssetRegistry data for the UObject.

OnAssetUpdated will provide the updated information a frame or two later, but ContentBrowserAssetDataSource does not have a hard dependency on up-to-date information.
#rnx

diff --git a/Engine/Plugins/Editor/ContentBrowser/ContentBrowserAssetDataSource/Source/ContentBrowserAssetDataSource/Private/ContentBrowserAssetDataSource.cpp b/Engine/Plugins/Editor/ContentBrowser/ContentBrowserAssetDataSource/Source/ContentBrowserAssetDataSource/Private/ContentBrowserAssetDataSource.cpp
--- a/Engine/Plugins/Editor/ContentBrowser/ContentBrowserAssetDataSource/Source/ContentBrowserAssetDataSource/Private/ContentBrowserAssetDataSource.cpp
+++ b/Engine/Plugins/Editor/ContentBrowser/ContentBrowserAssetDataSource/Source/ContentBrowserAssetDataSource/Private/ContentBrowserAssetDataSource.cpp
@@ -56,7 +56,6 @@ void UContentBrowserAssetDataSource::Initialize(const bool InAutoRegister)
 	AssetRegistry->OnFilesLoaded().AddUObject(this, &UContentBrowserAssetDataSource::OnScanCompleted);
 
 	// Listen for when assets are loaded or changed
-	FCoreUObjectDelegates::OnAssetLoaded.AddUObject(this, &UContentBrowserAssetDataSource::OnAssetLoaded);
 	FCoreUObjectDelegates::OnObjectPropertyChanged.AddUObject(this, &UContentBrowserAssetDataSource::OnObjectPropertyChanged);
 
 	// Listen for new mount roots
@@ -180,7 +179,6 @@ void UContentBrowserAssetDataSource::Shutdown()
 		AssetRegistry->OnPathRemoved().RemoveAll(this);
 	}
 
-	FCoreUObjectDelegates::OnAssetLoaded.RemoveAll(this);
 	FCoreUObjectDelegates::OnObjectPropertyChanged.RemoveAll(this);
 
 	AssetViewUtils::OnAlwaysShowPath().RemoveAll(this);
@@ -1662,19 +1660,6 @@ void UContentBrowserAssetDataSource::OnAssetUpdated(const FAssetData& InAssetDat
 	}
 }
 
-void UContentBrowserAssetDataSource::OnAssetLoaded(UObject* InAsset)
-{
-	if (InAsset && !InAsset->GetOutermost()->HasAnyPackageFlags(PKG_ForDiffing | PKG_PlayInEditor) &&
-		!UE::AssetRegistry::FFiltering::ShouldSkipAsset(InAsset))
-	{
-		FAssetData AssetData(InAsset);
-		if (ContentBrowserAssetData::IsPrimaryAsset(AssetData))
-		{
-			QueueItemDataUpdate(FContentBrowserItemDataUpdate::MakeItemModifiedUpdate(CreateAssetFileItem(AssetData)));
-		}
-	}
-}
-
 void UContentBrowserAssetDataSource::OnObjectPropertyChanged(UObject* InObject, FPropertyChangedEvent& InPropertyChangedEvent)
 {
 	if (InObject && InObject->IsAsset())
diff --git a/Engine/Plugins/Editor/ContentBrowser/ContentBrowserAssetDataSource/Source/ContentBrowserAssetDataSource/Public/ContentBrowserAssetDataSource.h b/Engine/Plugins/Editor/ContentBrowser/ContentBrowserAssetDataSource/Source/ContentBrowserAssetDataSource/Public/ContentBrowserAssetDataSource.h
--- a/Engine/Plugins/Editor/ContentBrowser/ContentBrowserAssetDataSource/Source/ContentBrowserAssetDataSource/Public/ContentBrowserAssetDataSource.h
+++ b/Engine/Plugins/Editor/ContentBrowser/ContentBrowserAssetDataSource/Source/ContentBrowserAssetDataSource/Public/ContentBrowserAssetDataSource.h
@@ -280,8 +280,6 @@ private:
 
 	void OnAssetUpdated(const FAssetData& InAssetData);
 
-	void OnAssetLoaded(UObject* InAsset);
-
 	void OnObjectPropertyChanged(UObject* InObject, FPropertyChangedEvent& InPropertyChangedEvent);
 
 	void OnPathAdded(const FString& InPath);
