Fix copy/paste error causing the 'View as text' context menu entry from being overriden by 'Run Text Asset Roundtrip' (found under  'Asset Actions->Text Assets')

#rnx

- Looks like a copy/paste error was introduced when the menu was refactored last year, which resulted in both commands being given the same name ('ViewSelectedAssetAsText')

diff --git a/Engine/Plugins/Editor/ContentBrowser/ContentBrowserAssetDataSource/Source/ContentBrowserAssetDataSource/Private/AssetFileContextMenu.cpp b/Engine/Plugins/Editor/ContentBrowser/ContentBrowserAssetDataSource/Source/ContentBrowserAssetDataSource/Private/AssetFileContextMenu.cpp
--- a/Engine/Plugins/Editor/ContentBrowser/ContentBrowserAssetDataSource/Source/ContentBrowserAssetDataSource/Private/AssetFileContextMenu.cpp
+++ b/Engine/Plugins/Editor/ContentBrowser/ContentBrowserAssetDataSource/Source/ContentBrowserAssetDataSource/Private/AssetFileContextMenu.cpp
@@ -554,7 +554,7 @@ void FAssetFileContextMenu::MakeAssetActionsSubMenu(UToolMenu* Menu)
 			);
 
 			FormatActionsSection.AddMenuEntry(
-				"ViewSelectedAssetAsText",
+				"RunTextAssetRoundtrip",
 				LOCTEXT("TextFormatRountrip", "Run Text Asset Roundtrip"),
 				LOCTEXT("TextFormatRountripTooltip", "Save the select asset backwards or forwards between text and binary formats and check for determinism"),
 				FSlateIcon(),
