Add a notification when assets can't be opened pointing the user to the detailed output log.

#rnx

diff --git a/Engine/Plugins/Editor/ContentBrowser/ContentBrowserAssetDataSource/Source/ContentBrowserAssetDataSource/Private/ContentBrowserAssetDataCore.cpp b/Engine/Plugins/Editor/ContentBrowser/ContentBrowserAssetDataSource/Source/ContentBrowserAssetDataSource/Private/ContentBrowserAssetDataCore.cpp
--- a/Engine/Plugins/Editor/ContentBrowser/ContentBrowserAssetDataSource/Source/ContentBrowserAssetDataSource/Private/ContentBrowserAssetDataCore.cpp
+++ b/Engine/Plugins/Editor/ContentBrowser/ContentBrowserAssetDataSource/Source/ContentBrowserAssetDataSource/Private/ContentBrowserAssetDataCore.cpp
@@ -231,6 +231,7 @@ bool EditOrPreviewAssetFileItems(TArrayView<const TSharedRef<const FContentBrows
 
 	// Now that we have created our map, load and activate all the lists of objects for each asset type action.
 	const bool bHasOpenActivationMethod = (ActivationMethod == EAssetTypeActivationMethod::DoubleClicked || ActivationMethod == EAssetTypeActivationMethod::Opened);
+	bool bSuccessfulEditorOpen = true;
 	for (auto& TypeActionToObjectsPair : TypeActionsToAssetData)
 	{
 		SlowTask.EnterProgressFrame(25.0f / TypeActionsToAssetData.Num());
@@ -265,11 +266,11 @@ bool EditOrPreviewAssetFileItems(TArrayView<const TSharedRef<const FContentBrows
 
 		if (bOpenEditorForAssets)
 		{
-			AssetEditorSubsystem->OpenEditorForAssets(ObjList);
+			bSuccessfulEditorOpen &= AssetEditorSubsystem->OpenEditorForAssets(ObjList);
 		}
 	}
 
-	return true;
+	return bSuccessfulEditorOpen;
 }
 
 bool EditOrPreviewItems(IAssetTools* InAssetTools, const UContentBrowserDataSource* InOwnerDataSource, TArrayView<const FContentBrowserItemData> InItems, const bool bIsPreview)
diff --git a/Engine/Source/Editor/ContentBrowser/Private/SContentBrowser.cpp b/Engine/Source/Editor/ContentBrowser/Private/SContentBrowser.cpp
--- a/Engine/Source/Editor/ContentBrowser/Private/SContentBrowser.cpp
+++ b/Engine/Source/Editor/ContentBrowser/Private/SContentBrowser.cpp
@@ -27,6 +27,7 @@
 #include "Widgets/Input/SSearchBox.h"
 #include "Widgets/Layout/SSplitter.h"
 #include "Widgets/Layout/SWidgetSwitcher.h"
+#include "Widgets/Notifications/SNotificationList.h"
 #include "Framework/Docking/TabManager.h"
 #include "EditorStyleSet.h"
 #include "EditorFontGlyphs.h"
@@ -2675,8 +2676,18 @@ void SContentBrowser::OnItemsActivated(TArrayView<const FContentBrowserItem> Act
 					AssetViewUtils::ShowErrorNotifcation(EditErrorMsg);
 				}
 			}
-
-			SourceAndItemsPair.Key->BulkEditItems(SourceAndItemsPair.Value);
+			
+			if (!SourceAndItemsPair.Key->BulkEditItems(SourceAndItemsPair.Value))
+			{
+				static const FText ErrorMessage = LOCTEXT("EditItemsFailure", "Failed to load assets");
+
+				FNotificationInfo WarningNotification(ErrorMessage);
+				WarningNotification.ExpireDuration = 5.0f;
+				WarningNotification.Hyperlink = FSimpleDelegate::CreateStatic([](){ FMessageLog("LoadErrors").Open(EMessageSeverity::Info, true); });
+				WarningNotification.HyperlinkText = LOCTEXT("LoadObjectHyperlink", "Show Message Log");
+				WarningNotification.bFireAndForget = true;
+				FSlateNotificationManager::Get().AddNotification(WarningNotification);
+			}
 		}
 	}
 }
diff --git a/Engine/Source/Editor/UnrealEd/Private/Subsystems/AssetEditorSubsystem.cpp b/Engine/Source/Editor/UnrealEd/Private/Subsystems/AssetEditorSubsystem.cpp
--- a/Engine/Source/Editor/UnrealEd/Private/Subsystems/AssetEditorSubsystem.cpp
+++ b/Engine/Source/Editor/UnrealEd/Private/Subsystems/AssetEditorSubsystem.cpp
@@ -553,7 +553,7 @@ bool UAssetEditorSubsystem::OpenEditorForAssets_Advanced(const TArray <UObject*
 		}
 	}
 
-	return true;
+	return NumNullAssets == 0;
 }
 
 bool UAssetEditorSubsystem::OpenEditorForAssets(const TArray<UObject*>& Assets)
