Allow blocking rename and duplicate UI per asset

diff --git a/Engine/Plugins/Editor/ContentBrowser/ContentBrowserAssetDataSource/Source/ContentBrowserAssetDataSource/Private/ContentBrowserAssetDataCore.cpp b/Engine/Plugins/Editor/ContentBrowser/ContentBrowserAssetDataSource/Source/ContentBrowserAssetDataSource/Private/ContentBrowserAssetDataCore.cpp
--- a/Engine/Plugins/Editor/ContentBrowser/ContentBrowserAssetDataSource/Source/ContentBrowserAssetDataSource/Private/ContentBrowserAssetDataCore.cpp
+++ b/Engine/Plugins/Editor/ContentBrowser/ContentBrowserAssetDataSource/Source/ContentBrowserAssetDataSource/Private/ContentBrowserAssetDataCore.cpp
@@ -327,6 +327,14 @@ bool CanDuplicateAssetFileItem(IAssetTools* InAssetTools, const FContentBrowserA
 		return false;
 	}
 
+	if (TSharedPtr<IAssetTypeActions> AssetTypeActions = InAssetPayload.GetAssetTypeActions())
+	{
+		if (!AssetTypeActions->CanDuplicate(InAssetPayload.GetAssetData(), OutErrorMsg))
+		{
+			return false;
+		}
+	}
+
 	return true;
 }
 
@@ -706,6 +714,14 @@ bool CanRenameAssetFileItem(IAssetTools* InAssetTools, const FContentBrowserAsse
 		}
 	}
 
+	if (TSharedPtr<IAssetTypeActions> AssetTypeActions = InAssetPayload.GetAssetTypeActions())
+	{
+		if (!AssetTypeActions->CanRename(InAssetPayload.GetAssetData(), OutErrorMsg))
+		{
+			return false;
+		}
+	}
+
 	return true;
 }
 
diff --git a/Engine/Source/Developer/AssetTools/Public/AssetTypeActions_Base.h b/Engine/Source/Developer/AssetTools/Public/AssetTypeActions_Base.h
--- a/Engine/Source/Developer/AssetTools/Public/AssetTypeActions_Base.h
+++ b/Engine/Source/Developer/AssetTools/Public/AssetTypeActions_Base.h
@@ -60,6 +60,16 @@ public:
 		return TArray<FAssetData>(InAssetDatas);
 	}
 
+	virtual bool CanRename(const FAssetData& InAsset, FText* OutErrorMsg) const override
+	{
+		return true;
+	}
+
+	virtual bool CanDuplicate(const FAssetData& InAsset, FText* OutErrorMsg) const override
+	{
+		return true;
+	}
+
 	virtual bool CanFilter() override
 	{
 		return true;
diff --git a/Engine/Source/Developer/AssetTools/Public/IAssetTypeActions.h b/Engine/Source/Developer/AssetTools/Public/IAssetTypeActions.h
--- a/Engine/Source/Developer/AssetTools/Public/IAssetTypeActions.h
+++ b/Engine/Source/Developer/AssetTools/Public/IAssetTypeActions.h
@@ -67,6 +67,12 @@ public:
 	/** Allows overriding asset activation to perform asset type specific activation for the supplied assets. This happens when the user double clicks, presses enter, or presses space. Return true if you have overridden the behavior. */
 	virtual bool AssetsActivatedOverride(const TArray<UObject*>& InObjects, EAssetTypeActivationMethod::Type ActivationType) = 0;
 
+	/** Returns true if this asset can be renamed */
+	virtual bool CanRename(const FAssetData& InAsset, FText* OutErrorMsg) const = 0;
+
+	/** Returns true if this asset can be duplicated */
+	virtual bool CanDuplicate(const FAssetData& InAsset, FText* OutErrorMsg) const = 0;
+
 	/** Returns the set of asset data that is valid to load. */
 	virtual TArray<FAssetData> GetValidAssetsForPreviewOrEdit(TArrayView<const FAssetData> InAssetDatas, bool bIsPreview) = 0;
 
