Improve user visibility for packages with virtualized data by adding "Has Virtualized Data" to the tool tip in the content browser as well as adding an asset filter to hide all packages without virtualized data.

#rnx

- The filter can be found in Filter->Other Filters -> Virtualized Data
- Since we only want to store if a package has virtualized data or not we can just use a bit in FAssetPackageData::Flags

diff --git a/Engine/Plugins/Editor/ContentBrowser/ContentBrowserAssetDataSource/Source/ContentBrowserAssetDataSource/Private/ContentBrowserAssetDataCore.cpp b/Engine/Plugins/Editor/ContentBrowser/ContentBrowserAssetDataSource/Source/ContentBrowserAssetDataSource/Private/ContentBrowserAssetDataCore.cpp
--- a/Engine/Plugins/Editor/ContentBrowser/ContentBrowserAssetDataSource/Source/ContentBrowserAssetDataSource/Private/ContentBrowserAssetDataCore.cpp
+++ b/Engine/Plugins/Editor/ContentBrowser/ContentBrowserAssetDataSource/Source/ContentBrowserAssetDataSource/Private/ContentBrowserAssetDataCore.cpp
@@ -1112,6 +1112,38 @@ bool GetDiskSizeItemAttribute(const FAssetData& InAssetData, IAssetRegistry* InA
 	return false;
 }
 
+bool GetVirtualizationItemAttribute(const FAssetData& InAssetData, IAssetRegistry* InAssetRegistry, const bool InIncludeMetaData, FContentBrowserItemDataAttributeValue& OutAttributeValue)
+{
+	check(InAssetData.IsValid());
+	check(InAssetRegistry != nullptr);
+
+	if (TOptional<FAssetPackageData> PackageData = InAssetRegistry->GetAssetPackageDataCopy(InAssetData.PackageName))
+	{
+		// We could set a bool here but that will display the value in lower case, where as asset properties 
+		// use a string value with the first letter in upper case, so we replicate that here to avoid the 
+		// entry looking out of place.
+		OutAttributeValue.SetValue(PackageData->HasVirtualizedPayloads() ? TEXT("True") : TEXT("False"));
+
+		if (InIncludeMetaData)
+		{
+			static const FText DisplayName = LOCTEXT("AttributeDisplayName_VirtualizedData", "Has Virtualized Data");
+
+			FContentBrowserItemDataAttributeMetaData AttributeMetaData;
+			AttributeMetaData.AttributeType = UObject::FAssetRegistryTag::TT_Alphabetical;
+			AttributeMetaData.DisplayFlags = UObject::FAssetRegistryTag::TD_None;
+			AttributeMetaData.DisplayName = DisplayName;
+
+			OutAttributeValue.SetMetaData(MoveTemp(AttributeMetaData));
+		}
+
+		return true;
+	}
+	else
+	{
+		return false;
+	}
+}
+
 void GetGenericItemAttribute(const FName InTagKey, const FString& InTagValue, const FAssetPropertyTagCache::FClassPropertyTagCache& InClassPropertyTagCache, const bool InIncludeMetaData, FContentBrowserItemDataAttributeValue& OutAttributeValue)
 {
 	check(!InTagKey.IsNone());
@@ -1248,6 +1280,11 @@ bool GetAssetFileItemAttribute(const FContentBrowserAssetFileItemDataPayload& In
 			return GetDiskSizeItemAttribute(InAssetPayload.GetAssetData(), IAssetRegistry::Get(), InIncludeMetaData, OutAttributeValue);
 		}
 
+		if (InAttributeKey == ContentBrowserItemAttributes::VirtualizedData)
+		{
+			return GetVirtualizationItemAttribute(InAssetPayload.GetAssetData(), IAssetRegistry::Get(), InIncludeMetaData, OutAttributeValue);
+		}
+
 		if (InAttributeKey == ContentBrowserItemAttributes::ItemIsDeveloperContent)
 		{
 			const bool bIsDevelopersFolder = AssetViewUtils::IsDevelopersFolder(InAssetPayload.GetAssetData().PackageName.ToString());
@@ -1345,6 +1382,13 @@ bool GetAssetFileItemAttributes(const FContentBrowserAssetFileItemDataPayload& I
 		{
 			OutAttributeValues.Add(ContentBrowserItemAttributes::ItemDiskSize, MoveTemp(DiskSizeAttributeValue));
 		}
+
+		// Virtualized Payloads
+		FContentBrowserItemDataAttributeValue VirtualizedAttributeValue;
+		if (GetVirtualizationItemAttribute(InAssetPayload.GetAssetData(), IAssetRegistry::Get(), InIncludeMetaData, VirtualizedAttributeValue))
+		{
+			OutAttributeValues.Add(ContentBrowserItemAttributes::VirtualizedData, MoveTemp(VirtualizedAttributeValue));
+		}	
 	}
 
 	// Generic attribute keys
diff --git a/Engine/Source/Editor/ContentBrowser/Private/FrontendFilters.cpp b/Engine/Source/Editor/ContentBrowser/Private/FrontendFilters.cpp
--- a/Engine/Source/Editor/ContentBrowser/Private/FrontendFilters.cpp
+++ b/Engine/Source/Editor/ContentBrowser/Private/FrontendFilters.cpp
@@ -26,6 +26,7 @@
 #include "TextFilterKeyValueHandlers.h"
 #include "TextFilterValueHandlers.h"
 #include "AssetCompilingManager.h"
+#include "ContentBrowserDataSource.h"
 
 /** Helper functions for frontend filters */
 namespace FrontendFilterHelper
@@ -1486,4 +1487,26 @@ bool FFrontendFilter_Writable::PassesFilter(FAssetFilterType InItem) const
 	ItemDiskPath = FPaths::ConvertRelativePathToFull(MoveTemp(ItemDiskPath));
 
 	return !IFileManager::Get().IsReadOnly(*ItemDiskPath);
+}
+
+/////////////////////////////////////////
+// FFrontendFilter_VirtualizedData
+/////////////////////////////////////////
+
+FFrontendFilter_VirtualizedData::FFrontendFilter_VirtualizedData(TSharedPtr<FFrontendFilterCategory> InCategory)
+	: FFrontendFilter(InCategory)
+{
+}
+
+bool FFrontendFilter_VirtualizedData::PassesFilter(FAssetFilterType InItem) const
+{
+	const FContentBrowserItemDataAttributeValue AttributeValue = InItem.GetItemAttribute(ContentBrowserItemAttributes::VirtualizedData);
+	if (AttributeValue.IsValid())
+	{
+		return AttributeValue.GetValue<FString>() == TEXT("True");
+	}
+	else
+	{
+		return false;
+	}
 }
\ No newline at end of file
diff --git a/Engine/Source/Editor/ContentBrowser/Private/SFilterList.cpp b/Engine/Source/Editor/ContentBrowser/Private/SFilterList.cpp
--- a/Engine/Source/Editor/ContentBrowser/Private/SFilterList.cpp
+++ b/Engine/Source/Editor/ContentBrowser/Private/SFilterList.cpp
@@ -553,6 +553,7 @@ void SFilterList::Construct( const FArguments& InArgs )
 	AllFrontendFilters.Add( MakeShareable(new FFrontendFilter_ArbitraryComparisonOperation(DefaultCategory)) );
 	AllFrontendFilters.Add(MakeShareable(new FFrontendFilter_Recent(DefaultCategory)));
 	AllFrontendFilters.Add( MakeShareable(new FFrontendFilter_NotSourceControlled(DefaultCategory)) );
+	AllFrontendFilters.Add(MakeShareable(new FFrontendFilter_VirtualizedData(DefaultCategory)));
 
 	// Add any global user-defined frontend filters
 	for (TObjectIterator<UContentBrowserFrontEndFilterExtension> ExtensionIt(RF_NoFlags); ExtensionIt; ++ExtensionIt)
diff --git a/Engine/Source/Editor/ContentBrowser/Public/FrontendFilters.h b/Engine/Source/Editor/ContentBrowser/Public/FrontendFilters.h
--- a/Engine/Source/Editor/ContentBrowser/Public/FrontendFilters.h
+++ b/Engine/Source/Editor/ContentBrowser/Public/FrontendFilters.h
@@ -390,4 +390,21 @@ public:
 private:
 };
 
+/** A filter that displays only packages that contain virtualized data  */
+class CONTENTBROWSER_API FFrontendFilter_VirtualizedData : public FFrontendFilter
+{
+public:
+	FFrontendFilter_VirtualizedData(TSharedPtr<FFrontendFilterCategory> InCategory);
+	~FFrontendFilter_VirtualizedData() = default;
+
+private:
+	// FFrontendFilter implementation
+	virtual FString GetName() const override { return TEXT("VirtualizedData"); }
+	virtual FText GetDisplayName() const override { return LOCTEXT("FrontendFilter_VirtualizedData", "Virtualized Data"); }
+	virtual FText GetToolTipText() const override { return LOCTEXT("FrontendFilter_VirtualizedDataTooltip", "Show only package that contain virtualized data."); }
+
+	// IFilter implementation
+	virtual bool PassesFilter(FAssetFilterType InItem) const override;
+};
+
 #undef LOCTEXT_NAMESPACE
diff --git a/Engine/Source/Editor/ContentBrowserData/Public/ContentBrowserDataSource.h b/Engine/Source/Editor/ContentBrowserData/Public/ContentBrowserDataSource.h
--- a/Engine/Source/Editor/ContentBrowserData/Public/ContentBrowserDataSource.h
+++ b/Engine/Source/Editor/ContentBrowserData/Public/ContentBrowserDataSource.h
@@ -42,6 +42,12 @@ namespace ContentBrowserItemAttributes
 	 */
 	const FName ItemDiskSize = "ItemDiskSize";
 
+	/**
+	 * Attribute key that can be used to query if the item has virtualized data or not.
+	 * Type: bool.
+	 */
+	const FName VirtualizedData = "HasVirtualizedData";
+
 	/**
 	 * Attribute key that can be used to query whether the given item is considered to be developer content.
 	 * Type: bool.
diff --git a/Engine/Source/Runtime/AssetRegistry/Private/PackageReader.cpp b/Engine/Source/Runtime/AssetRegistry/Private/PackageReader.cpp
--- a/Engine/Source/Runtime/AssetRegistry/Private/PackageReader.cpp
+++ b/Engine/Source/Runtime/AssetRegistry/Private/PackageReader.cpp
@@ -11,6 +11,7 @@
 #include "Misc/Guid.h"
 #include "Misc/PackageName.h"
 #include "UObject/Class.h"
+#include "UObject/PackageTrailer.h"
 
 FPackageReader::FPackageReader()
 	: Loader(nullptr)
@@ -237,6 +238,37 @@ bool FPackageReader::SerializeAssetRegistryDependencyData(FPackageDependencyData
 	return true;
 }
 
+bool FPackageReader::SerializePackageTrailer(FPackageDependencyData& DependencyData)
+{
+	bool bHasVirtualizedPayloads = false;
+
+	if (PackageFileSummary.PayloadTocOffset != INDEX_NONE)
+	{
+		// Store the original offset in the package being read, since we are about to seek to near the end
+		const int64 InitialOffset = Tell();
+		
+		Seek(PackageFileSummary.PayloadTocOffset);
+
+		if (IsError())
+		{
+			return false;
+		}
+
+		UE::FPackageTrailer Trailer;
+		if (Trailer.TryLoad(*this))
+		{
+			bHasVirtualizedPayloads = Trailer.GetNumPayloads(UE::EPayloadFilter::Virtualized) > 0 ? true : false;
+		}
+
+		// Restore the original offset inside the package
+		Seek(InitialOffset);
+	}
+
+	DependencyData.PackageData.SetHasVirtualizedPayloads(bHasVirtualizedPayloads);
+
+	return true;
+}
+
 bool FPackageReader::ReadAssetDataFromThumbnailCache(TArray<FAssetData*>& AssetDataList)
 {
 	if (!StartSerializeSection(PackageFileSummary.ThumbnailTableOffset))
@@ -412,6 +444,11 @@ bool FPackageReader::ReadDependencyData(FPackageDependencyData& OutDependencyDat
 	{
 		return false;
 	}
+	// Note that the trailer is read from the end of the package file, and so should be the last thing read here
+	if (!SerializePackageTrailer(OutDependencyData))
+	{
+		return false;
+	}
 
 	checkf(OutDependencyData.IsValid(), TEXT("We should have early exited above rather than creating invalid dependency data"));
 	return true;
diff --git a/Engine/Source/Runtime/AssetRegistry/Private/PackageReader.h b/Engine/Source/Runtime/AssetRegistry/Private/PackageReader.h
--- a/Engine/Source/Runtime/AssetRegistry/Private/PackageReader.h
+++ b/Engine/Source/Runtime/AssetRegistry/Private/PackageReader.h
@@ -61,6 +61,7 @@ public:
 	bool SerializeSoftPackageReferenceList(TArray<FName>& OutSoftPackageReferenceList);
 	bool SerializeSearchableNamesMap(FPackageDependencyData& OutDependencyData);
 	bool SerializeAssetRegistryDependencyData(FPackageDependencyData& DependencyData);
+	bool SerializePackageTrailer(FPackageDependencyData& DependencyData);
 
 	/** Returns flags the asset package was saved with */
 	uint32 GetPackageFlags() const;
diff --git a/Engine/Source/Runtime/CoreUObject/Public/AssetRegistry/AssetData.h b/Engine/Source/Runtime/CoreUObject/Public/AssetRegistry/AssetData.h
--- a/Engine/Source/Runtime/CoreUObject/Public/AssetRegistry/AssetData.h
+++ b/Engine/Source/Runtime/CoreUObject/Public/AssetRegistry/AssetData.h
@@ -740,6 +740,10 @@ public:
 	bool IsLicenseeVersion() const { return (Flags & FLAG_LICENSEE_VERSION) != 0; }
 	void SetIsLicenseeVersion(bool bValue) { Flags = (Flags & ~FLAG_LICENSEE_VERSION) | (bValue ? FLAG_LICENSEE_VERSION : 0); }
 
+	/** Whether the package contains virtualized payloads or not */
+	bool HasVirtualizedPayloads() const { return (Flags & FLAG_HAS_VIRTUALIZED_PAYLOADS) != 0; }
+	void SetHasVirtualizedPayloads(bool bValue) { Flags = (Flags & ~FLAG_HAS_VIRTUALIZED_PAYLOADS) | (bValue ? FLAG_HAS_VIRTUALIZED_PAYLOADS : 0); }
+
 	/**
 	 * Serialize as part of the registry cache. This is not meant to be serialized as part of a package so  it does not handle versions normally
 	 * To version this data change FAssetRegistryVersion
@@ -757,7 +761,8 @@ public:
 private:
 	enum
 	{
-		FLAG_LICENSEE_VERSION = 0x1,
+		FLAG_LICENSEE_VERSION			= 1 << 0,
+		FLAG_HAS_VIRTUALIZED_PAYLOADS	= 1 << 1
 	};
 
 	void SerializeForCacheInternal(FArchive& Ar, FAssetRegistryVersion::Type Version);
