- FLinkerInstancingContext: Add Tag support
- Use FLinkerInstancingContext Tag to load all external actors instead of using a global map

#rnx

diff --git a/Engine/Plugins/Editor/ContentBrowser/ContentBrowserAssetDataSource/Source/ContentBrowserAssetDataSource/Private/ContentBrowserAssetDataCore.cpp b/Engine/Plugins/Editor/ContentBrowser/ContentBrowserAssetDataSource/Source/ContentBrowserAssetDataSource/Private/ContentBrowserAssetDataCore.cpp
--- a/Engine/Plugins/Editor/ContentBrowser/ContentBrowserAssetDataSource/Source/ContentBrowserAssetDataSource/Private/ContentBrowserAssetDataCore.cpp
+++ b/Engine/Plugins/Editor/ContentBrowser/ContentBrowserAssetDataSource/Source/ContentBrowserAssetDataSource/Private/ContentBrowserAssetDataCore.cpp
@@ -18,7 +18,7 @@
 #include "AssetFileContextMenu.h"
 #include "ContentBrowserDataUtils.h"
 #include "Settings/ContentBrowserSettings.h"
-#include "Engine/World.h"
+#include "Engine/Level.h"
 
 #define LOCTEXT_NAMESPACE "ContentBrowserAssetDataSource"
 
@@ -355,8 +355,7 @@ bool DuplicateItem(IAssetTools* InAssetTools, const UContentBrowserDataSource* I
 bool DuplicateAssetFileItem(IAssetTools* InAssetTools, const FContentBrowserAssetFileItemDataPayload& InAssetPayload, UObject*& OutSourceAsset, FAssetData& OutNewAsset)
 {
 	// We need to potentially load the asset in order to duplicate it
-	FScopedLoadAllExternalObjects Scope(InAssetPayload.GetAssetData().PackageName);
-	if (UObject* Asset = InAssetPayload.LoadAsset())
+	if (UObject* Asset = InAssetPayload.LoadAsset({ ULevel::LoadAllExternalObjectsTag }))
 	{
 		// Find a unique default name for the duplicated asset
 		FString DefaultAssetName;
@@ -394,8 +393,7 @@ bool DuplicateAssetFileItems(TArrayView<const TSharedRef<const FContentBrowserAs
 	for (const TSharedRef<const FContentBrowserAssetFileItemDataPayload>& AssetPayload : InAssetPayloads)
 	{
 		// We need to potentially load the asset in order to duplicate it
-		FScopedLoadAllExternalObjects Scope(AssetPayload->GetAssetData().PackageName);
-		if (UObject* Asset = AssetPayload->LoadAsset())
+		if (UObject* Asset = AssetPayload->LoadAsset({ ULevel::LoadAllExternalObjectsTag }))
 		{
 			ObjectsToDuplicate.Add(Asset);
 		}
@@ -770,8 +768,7 @@ bool RenameAssetFolderItem(IAssetRegistry* InAssetRegistry, const FContentBrowse
 bool RenameAssetFileItem(IAssetTools* InAssetTools, const FContentBrowserAssetFileItemDataPayload& InAssetPayload, const FString& InNewName)
 {
 	// We need to potentially load the asset in order to rename it
-	FScopedLoadAllExternalObjects Scope(InAssetPayload.GetAssetData().PackageName);
-	if (UObject* Asset = InAssetPayload.LoadAsset())
+	if (UObject* Asset = InAssetPayload.LoadAsset({ ULevel::LoadAllExternalObjectsTag }))
 	{
 		const FString PackagePath = FPackageName::GetLongPackagePath(Asset->GetOutermost()->GetName());
 
@@ -844,8 +841,7 @@ bool CopyAssetFileItems(TArrayView<const TSharedRef<const FContentBrowserAssetFi
 	for (const TSharedRef<const FContentBrowserAssetFileItemDataPayload>& AssetPayload : InAssetPayloads)
 	{
 		// We need to potentially load the asset in order to duplicate it
-		FScopedLoadAllExternalObjects Scope(AssetPayload->GetAssetData().PackageName);
-		if (UObject* Asset = AssetPayload->LoadAsset())
+		if (UObject* Asset = AssetPayload->LoadAsset({ ULevel::LoadAllExternalObjectsTag }))
 		{
 			AssetsToCopy.Add(Asset);
 		}
@@ -930,8 +926,7 @@ bool MoveAssetFileItems(TArrayView<const TSharedRef<const FContentBrowserAssetFi
 	for (const TSharedRef<const FContentBrowserAssetFileItemDataPayload>& AssetPayload : InAssetPayloads)
 	{
 		// We need to potentially load the asset in order to duplicate it
-		FScopedLoadAllExternalObjects Scope(AssetPayload->GetAssetData().PackageName);
-		if (UObject* Asset = AssetPayload->LoadAsset())
+		if (UObject* Asset = AssetPayload->LoadAsset({ ULevel::LoadAllExternalObjectsTag }))
 		{
 			AssetsToMove.Add(Asset);
 		}
diff --git a/Engine/Plugins/Editor/ContentBrowser/ContentBrowserAssetDataSource/Source/ContentBrowserAssetDataSource/Private/ContentBrowserAssetDataPayload.cpp b/Engine/Plugins/Editor/ContentBrowser/ContentBrowserAssetDataSource/Source/ContentBrowserAssetDataSource/Private/ContentBrowserAssetDataPayload.cpp
--- a/Engine/Plugins/Editor/ContentBrowser/ContentBrowserAssetDataSource/Source/ContentBrowserAssetDataSource/Private/ContentBrowserAssetDataPayload.cpp
+++ b/Engine/Plugins/Editor/ContentBrowser/ContentBrowserAssetDataSource/Source/ContentBrowserAssetDataSource/Private/ContentBrowserAssetDataPayload.cpp
@@ -41,13 +41,14 @@ UPackage* FContentBrowserAssetFileItemDataPayload::GetPackage(const bool bTryRec
 	return CachedPackagePtr.Get();
 }
 
-UPackage* FContentBrowserAssetFileItemDataPayload::LoadPackage() const
+UPackage* FContentBrowserAssetFileItemDataPayload::LoadPackage(TSet<FName> LoadTags) const
 {
 	if (!bHasCachedPackagePtr || !CachedPackagePtr.IsValid())
 	{
 		if (!AssetData.PackageName.IsNone())
 		{
-			CachedPackagePtr = ::LoadPackage(nullptr, *FNameBuilder(AssetData.PackageName), LOAD_None);
+			FLinkerInstancingContext InstancingContext(MoveTemp(LoadTags));
+			CachedPackagePtr = ::LoadPackage(nullptr, *FNameBuilder(AssetData.PackageName), LOAD_None, nullptr, &InstancingContext);
 			(void)GetAsset(/*bTryRecacheIfNull*/true); // Also re-cache the asset pointer
 		}
 		bHasCachedPackagePtr = true;
@@ -68,13 +69,14 @@ UObject* FContentBrowserAssetFileItemDataPayload::GetAsset(const bool bTryRecach
 	return CachedAssetPtr.Get();
 }
 
-UObject* FContentBrowserAssetFileItemDataPayload::LoadAsset() const
+UObject* FContentBrowserAssetFileItemDataPayload::LoadAsset(TSet<FName> LoadTags) const
 {
 	if (!bHasCachedAssetPtr || !CachedAssetPtr.IsValid())
 	{
 		if (!AssetData.ObjectPath.IsNone())
 		{
-			CachedAssetPtr = LoadObject<UObject>(nullptr, *FNameBuilder(AssetData.ObjectPath));
+			FLinkerInstancingContext InstancingContext(MoveTemp(LoadTags));
+			CachedAssetPtr = LoadObject<UObject>(nullptr, *FNameBuilder(AssetData.ObjectPath), nullptr, LOAD_None, nullptr, &InstancingContext);
 			(void)GetPackage(/*bTryRecacheIfNull*/true); // Also re-cache the package pointer
 		}
 		bHasCachedAssetPtr = true;
diff --git a/Engine/Plugins/Editor/ContentBrowser/ContentBrowserAssetDataSource/Source/ContentBrowserAssetDataSource/Public/ContentBrowserAssetDataPayload.h b/Engine/Plugins/Editor/ContentBrowser/ContentBrowserAssetDataSource/Source/ContentBrowserAssetDataSource/Public/ContentBrowserAssetDataPayload.h
--- a/Engine/Plugins/Editor/ContentBrowser/ContentBrowserAssetDataSource/Source/ContentBrowserAssetDataSource/Public/ContentBrowserAssetDataPayload.h
+++ b/Engine/Plugins/Editor/ContentBrowser/ContentBrowserAssetDataSource/Source/ContentBrowserAssetDataSource/Public/ContentBrowserAssetDataPayload.h
@@ -6,6 +6,7 @@
 #include "ContentBrowserItemData.h"
 #include "AssetData.h"
 #include "UObject/GCObject.h"
+#include "Containers/Set.h"
 
 class IAssetTypeActions;
 class FAssetThumbnail;
@@ -46,11 +47,13 @@ public:
 
 	UPackage* GetPackage(const bool bTryRecacheIfNull = false) const;
 
-	UPackage* LoadPackage() const;
+	// LoadTags(optional) allows passing specific tags to the linker when loading the asset package (@see ULevel::LoadAllExternalObjectsTag for an example usage)
+	UPackage* LoadPackage(TSet<FName> LoadTags = {}) const;
 
 	UObject* GetAsset(const bool bTryRecacheIfNull = false) const;
 
-	UObject* LoadAsset() const;
+	//  LoadTags (optional) allows passing specific tags to the linker when loading the asset (@see ULevel::LoadAllExternalObjectsTag for an example usage)
+	UObject* LoadAsset(TSet<FName> LoadTags = {}) const;
 
 	TSharedPtr<IAssetTypeActions> GetAssetTypeActions() const;
 
diff --git a/Engine/Source/Developer/AssetTools/Private/AssetViewUtils.cpp b/Engine/Source/Developer/AssetTools/Private/AssetViewUtils.cpp
--- a/Engine/Source/Developer/AssetTools/Private/AssetViewUtils.cpp
+++ b/Engine/Source/Developer/AssetTools/Private/AssetViewUtils.cpp
@@ -39,6 +39,7 @@
 #include "Interfaces/IPluginManager.h"
 #include "Subsystems/AssetEditorSubsystem.h"
 #include "Editor.h"
+#include "UObject/LinkerInstancingContext.h"
 
 #define LOCTEXT_NAMESPACE "ContentBrowser"
 
@@ -177,8 +178,8 @@ bool AssetViewUtils::LoadAssetsIfNeeded(const TArray<FString>& ObjectPaths, TArr
 			SlowTask.EnterProgressFrame(1, FText::Format(LOCTEXT("LoadingObjectf", "Loading {0}..."), FText::FromString(ObjectPath)));
 
 			// Load up the object
-			TUniquePtr<FScopedLoadAllExternalObjects> Scope(bLoadAllExternalObjects? new FScopedLoadAllExternalObjects(*FEditorFileUtils::ExtractPackageName(ObjectPath)) : nullptr);
-			UObject* LoadedObject = LoadObject<UObject>(NULL, *ObjectPath, NULL, LoadFlags, NULL);
+			FLinkerInstancingContext InstancingContext(bLoadAllExternalObjects ? TSet<FName>{ ULevel::LoadAllExternalObjectsTag } : TSet<FName>());
+			UObject* LoadedObject = LoadObject<UObject>(NULL, *ObjectPath, NULL, LoadFlags, NULL, &InstancingContext);
 			if ( LoadedObject )
 			{
 				LoadedObjects.Add(LoadedObject);
@@ -991,8 +992,7 @@ void AssetViewUtils::GetObjectsInAssetData(const TArray<FAssetData>& AssetList,
 	for (int32 AssetIdx = 0; AssetIdx < AssetList.Num(); ++AssetIdx)
 	{
 		const FAssetData& AssetData = AssetList[AssetIdx];
-		TUniquePtr<FScopedLoadAllExternalObjects> LoadAllExternalObjects(bLoadAllExternalObjects ? new FScopedLoadAllExternalObjects(AssetData.PackageName) : nullptr);
-		UObject* Obj = AssetData.GetAsset();
+		UObject* Obj = AssetData.GetAsset(bLoadAllExternalObjects ? TSet<FName> { ULevel::LoadAllExternalObjectsTag } : TSet<FName>());
 		if (Obj)
 		{
 			OutDroppedObjects.Add(Obj);
diff --git a/Engine/Source/Editor/ContentBrowser/Private/SAssetView.cpp b/Engine/Source/Editor/ContentBrowser/Private/SAssetView.cpp
--- a/Engine/Source/Editor/ContentBrowser/Private/SAssetView.cpp
+++ b/Engine/Source/Editor/ContentBrowser/Private/SAssetView.cpp
@@ -25,6 +25,7 @@
 #include "EditorStyleSet.h"
 #include "Settings/ContentBrowserSettings.h"
 #include "Engine/Blueprint.h"
+#include "Engine/Level.h"
 #include "Editor.h"
 #include "AssetSelection.h"
 #include "IAssetTools.h"
@@ -1582,8 +1583,8 @@ FReply SAssetView::OnKeyDown( const FGeometry& MyGeometry, const FKeyEvent& InKe
 				FString PackageName;
 				if (IsValidObjectPath(AssetPath, ObjectClassName, ObjectPath, PackageName))
 				{
-					FScopedLoadAllExternalObjects Scope(*PackageName);
-					UObject* ObjectToCopy = LoadObject<UObject>(nullptr, *ObjectPath);
+					FLinkerInstancingContext InstancingContext({ ULevel::LoadAllExternalObjectsTag });
+					UObject* ObjectToCopy = LoadObject<UObject>(nullptr, *ObjectPath, nullptr, LOAD_None, nullptr, &InstancingContext);
 					if (ObjectToCopy && !ObjectToCopy->IsA(UClass::StaticClass()))
 					{
 						AssetsToCopy.Add(ObjectToCopy);
diff --git a/Engine/Source/Editor/UnrealEd/Private/DragAndDrop/AssetDragDropOp.cpp b/Engine/Source/Editor/UnrealEd/Private/DragAndDrop/AssetDragDropOp.cpp
--- a/Engine/Source/Editor/UnrealEd/Private/DragAndDrop/AssetDragDropOp.cpp
+++ b/Engine/Source/Editor/UnrealEd/Private/DragAndDrop/AssetDragDropOp.cpp
@@ -9,6 +9,7 @@
 #include "AssetThumbnail.h"
 #include "ClassIconFinder.h"
 #include "ThumbnailRendering/ThumbnailManager.h"
+#include "Engine/Level.h"
 
 TSharedRef<FAssetDragDropOp> FAssetDragDropOp::New(const FAssetData& InAssetData, UActorFactory* ActorFactory)
 {
@@ -210,8 +211,7 @@ void FAssetDragDropOp::Init(TArray<FAssetData> InAssetData, TArray<FString> InAs
 	// Can cause unsafe frame reentry 
 	for (FAssetData& Data : AssetData)
 	{
-		FScopedLoadAllExternalObjects Scope(Data.PackageName);
-		Data.GetAsset();
+		Data.GetAsset({ ULevel::LoadAllExternalObjectsTag });
 	}
 
 	InitThumbnail();
diff --git a/Engine/Source/Editor/UnrealEd/Private/ObjectTools.cpp b/Engine/Source/Editor/UnrealEd/Private/ObjectTools.cpp
--- a/Engine/Source/Editor/UnrealEd/Private/ObjectTools.cpp
+++ b/Engine/Source/Editor/UnrealEd/Private/ObjectTools.cpp
@@ -2469,8 +2469,7 @@ namespace ObjectTools
 		for ( int i = 0; i < AssetsToDelete.Num(); i++ )
 		{
 			const FAssetData& AssetData = AssetsToDelete[i];
-			FScopedLoadAllExternalObjects Scope(AssetData.PackageName);
-			UObject *ObjectToDelete = AssetData.GetAsset();
+			UObject *ObjectToDelete = AssetData.GetAsset({ ULevel::LoadAllExternalObjectsTag });
 			// Assets can be loaded even when their underlying type/class no longer exists...
 			if ( ObjectToDelete!=nullptr )
 			{
diff --git a/Engine/Source/Runtime/CoreUObject/Public/AssetRegistry/AssetData.h b/Engine/Source/Runtime/CoreUObject/Public/AssetRegistry/AssetData.h
--- a/Engine/Source/Runtime/CoreUObject/Public/AssetRegistry/AssetData.h
+++ b/Engine/Source/Runtime/CoreUObject/Public/AssetRegistry/AssetData.h
@@ -8,6 +8,7 @@
 #include "AssetRegistry/AssetDataTagMap.h"
 #include "Containers/ArrayView.h"
 #include "Containers/StringView.h"
+#include "Containers/Set.h"
 #include "Misc/PackageName.h"
 #include "Misc/StringBuilder.h"
 #include "Templates/UniquePtr.h"
@@ -19,6 +20,7 @@
 #include "UObject/ObjectRedirector.h"
 #include "UObject/Package.h"
 #include "UObject/PrimaryAssetId.h"
+#include "UObject/LinkerInstancingContext.h"
 
 struct FCustomVersion;
 
@@ -317,8 +319,13 @@ public:
 	/** Gets primary asset id of this data */
 	COREUOBJECT_API FPrimaryAssetId GetPrimaryAssetId() const;
 
-	/** Returns the asset UObject if it is loaded or loads the asset if it is unloaded then returns the result */
-	UObject* FastGetAsset(bool bLoad=false) const
+	/** 
+	 * Returns the asset UObject if it is loaded or loads the asset if it is unloaded then returns the result
+	 * 
+	 * @param bLoad (optional) loads the asset if it is unloaded.
+	 * @param LoadTags (optional) allows passing specific tags to the linker when loading the asset (@see ULevel::LoadAllExternalObjectsTag for an example usage)
+	 */
+	UObject* FastGetAsset(bool bLoad = false, TSet<FName> LoadTags = {}) const
 	{
 		if ( !IsValid())
 		{
@@ -331,7 +338,8 @@ public:
 		{
 			if (bLoad)
 			{
-				return LoadObject<UObject>(nullptr, *ObjectPath.ToString());
+				FLinkerInstancingContext InstancingContext(MoveTemp(LoadTags));
+				return LoadObject<UObject>(nullptr, *ObjectPath.ToString(), nullptr, 0, nullptr, &InstancingContext);
 			}
 			else
 			{
@@ -348,8 +356,12 @@ public:
 		return Asset;
 	}
 
-	/** Returns the asset UObject if it is loaded or loads the asset if it is unloaded then returns the result */
-	UObject* GetAsset() const
+	/** 
+	 * Returns the asset UObject if it is loaded or loads the asset if it is unloaded then returns the result 
+	 * 
+	 * @param LoadTags (optional) allows passing specific tags to the linker when loading the asset (@see ULevel::LoadAllExternalObjectsTag for an example usage)
+	 */
+	UObject* GetAsset(TSet<FName> LoadTags = {}) const
 	{
 		if ( !IsValid())
 		{
@@ -360,7 +372,8 @@ public:
 		UObject* Asset = FindObject<UObject>(nullptr, *ObjectPath.ToString());
 		if ( Asset == nullptr)
 		{
-			Asset = LoadObject<UObject>(nullptr, *ObjectPath.ToString());
+			FLinkerInstancingContext InstancingContext(MoveTemp(LoadTags));
+			Asset = LoadObject<UObject>(nullptr, *ObjectPath.ToString(), nullptr, 0, nullptr, &InstancingContext);
 		}
 
 		return Asset;
diff --git a/Engine/Source/Runtime/CoreUObject/Public/UObject/LinkerInstancingContext.h b/Engine/Source/Runtime/CoreUObject/Public/UObject/LinkerInstancingContext.h
--- a/Engine/Source/Runtime/CoreUObject/Public/UObject/LinkerInstancingContext.h
+++ b/Engine/Source/Runtime/CoreUObject/Public/UObject/LinkerInstancingContext.h
@@ -15,6 +15,10 @@ public:
 	FLinkerInstancingContext(TMap<FName, FName> InInstanceMapping)
 		: Mapping(MoveTemp(InInstanceMapping))
 	{}
+	FLinkerInstancingContext(TSet<FName> InTags)
+		: Tags(MoveTemp(InTags))
+	{
+	}
 
 	bool IsInstanced() const
 	{
@@ -41,6 +45,21 @@ public:
 		Mapping.Append(NewMapping);
 	}
 
+	void AddTag(FName NewTag)
+	{
+		Tags.Add(NewTag);
+	}
+
+	void AppendTags(const TSet<FName>& NewTags)
+	{
+		Tags.Append(NewTags);
+	}
+
+	bool HasTag(FName Tag) const
+	{
+		return Tags.Contains(Tag);
+	}
+
 	/** Return the instanced package name for a given instanced outer package and an object package name */
 	static FString GetInstancedPackageName(const FString& InOuterPackageName, const FString& InPackageName)
 	{
@@ -52,4 +71,6 @@ private:
 
 	/** Map of original object name to their instance counterpart. */
 	TMap<FName, FName> Mapping;
+	/** Tags can be used to determine some loading behavior. */
+	TSet<FName> Tags;
 };
diff --git a/Engine/Source/Runtime/CoreUObject/Public/UObject/UObjectGlobals.h b/Engine/Source/Runtime/CoreUObject/Public/UObject/UObjectGlobals.h
--- a/Engine/Source/Runtime/CoreUObject/Public/UObject/UObjectGlobals.h
+++ b/Engine/Source/Runtime/CoreUObject/Public/UObject/UObjectGlobals.h
@@ -1541,9 +1541,9 @@ inline T* FindObjectSafe( UObject* Outer, const TCHAR* Name, bool ExactClass=fal
  * @see StaticLoadObject()
  */
 template< class T > 
-inline T* LoadObject( UObject* Outer, const TCHAR* Name, const TCHAR* Filename=nullptr, uint32 LoadFlags=LOAD_None, UPackageMap* Sandbox=nullptr )
+inline T* LoadObject( UObject* Outer, const TCHAR* Name, const TCHAR* Filename=nullptr, uint32 LoadFlags=LOAD_None, UPackageMap* Sandbox=nullptr, const FLinkerInstancingContext* InstancingContext=nullptr )
 {
-	return (T*)StaticLoadObject( T::StaticClass(), Outer, Name, Filename, LoadFlags, Sandbox );
+	return (T*)StaticLoadObject( T::StaticClass(), Outer, Name, Filename, LoadFlags, Sandbox, true, InstancingContext );
 }
 
 /**
diff --git a/Engine/Source/Runtime/Engine/Classes/Engine/Level.h b/Engine/Source/Runtime/Engine/Classes/Engine/Level.h
--- a/Engine/Source/Runtime/Engine/Classes/Engine/Level.h
+++ b/Engine/Source/Runtime/Engine/Classes/Engine/Level.h
@@ -779,7 +779,7 @@ public:
 	ENGINE_API static bool GetIsStreamingDisabledFromPackage(FName LevelPackage);
 	ENGINE_API static bool GetPartitionedLevelCanBeUsedByLevelInstanceFromAsset(const FAssetData& Asset);
 	ENGINE_API static bool GetPartitionedLevelCanBeUsedByLevelInstanceFromPackage(FName LevelPackage);
-
+	ENGINE_API static const FName LoadAllExternalObjectsTag;
 	ENGINE_API bool GetPromptWhenAddingToLevelOutsideBounds() const;
 	ENGINE_API bool GetPromptWhenAddingToLevelBeforeCheckout() const;
 #endif
diff --git a/Engine/Source/Runtime/Engine/Classes/Engine/World.h b/Engine/Source/Runtime/Engine/Classes/Engine/World.h
--- a/Engine/Source/Runtime/Engine/Classes/Engine/World.h
+++ b/Engine/Source/Runtime/Engine/Classes/Engine/World.h
@@ -1001,17 +1001,6 @@ struct FWorldCachedViewInfo
 	FMatrix ViewToWorld;
 };
 
-/**
- * Scope can be used to force load all external objects when loading a world. 
- */
-struct ENGINE_API FScopedLoadAllExternalObjects
-{
-	FScopedLoadAllExternalObjects(FName InPackageName);
-	~FScopedLoadAllExternalObjects();
-
-	FName PackageName;
-};
-
 /**
  * Helper class allows UWorldPartition to broadcast UWorld events 
  */
@@ -1526,12 +1515,7 @@ public:
 	/** Returns whether or not this world is currently ticking. See SetShouldTick. */
 	bool ShouldTick() const { return bShouldTick; }
 
-	static bool ShouldLoadAllExternalObjects(FName InPackageName) { return LoadAllExternalObjects.Contains(InPackageName); }
-
 private:
-	friend struct FScopedLoadAllExternalObjects;
-	static TSet<FName> LoadAllExternalObjects;
-
 	/** List of all the controllers in the world. */
 	TArray<TWeakObjectPtr<class AController> > ControllerList;
 
diff --git a/Engine/Source/Runtime/Engine/Private/Level.cpp b/Engine/Source/Runtime/Engine/Private/Level.cpp
--- a/Engine/Source/Runtime/Engine/Private/Level.cpp
+++ b/Engine/Source/Runtime/Engine/Private/Level.cpp
@@ -76,6 +76,7 @@ Level.cpp: Level-related functions
 #include "ScopedTransaction.h"
 #include "EditorActorFolders.h"
 #include "UObject/MetaData.h"
+#include "UObject/LinkerLoad.h"
 #endif
 #include "WorldPartition/WorldPartition.h"
 #include "WorldPartition/WorldPartitionSubsystem.h"
@@ -363,6 +364,10 @@ static bool IsActorFolderObjectsFeatureAvailable()
 
 TMap<FName, TWeakObjectPtr<UWorld> > ULevel::StreamedLevelsOwningWorld;
 
+#if WITH_EDITOR
+const FName ULevel::LoadAllExternalObjectsTag(TEXT("LoadAllExternalObjectsTag"));
+#endif
+
 ULevel::ULevel( const FObjectInitializer& ObjectInitializer )
 	:	UObject( ObjectInitializer )
 	,	Actors()
@@ -920,10 +925,13 @@ void ULevel::PostLoad()
 		}
 	}
 
+	const FLinkerLoad* Linker = GetLinker();
+	check(Linker || bWasDuplicated);
+
 	// if we use external actors, load dynamic actors here
 	if (IsUsingExternalActors())
 	{
-		if (!bWasDuplicated && (!bIsPartitioned || UWorld::ShouldLoadAllExternalObjects(GetPackage()->GetFName())))
+		if (!bWasDuplicated && (!bIsPartitioned || Linker->GetInstancingContext().HasTag(ULevel::LoadAllExternalObjectsTag)))
 		{
 			UPackage* LevelPackage = GetPackage();
 			bool bPackageForPIE = LevelPackage->HasAnyPackageFlags(PKG_PlayInEditor);
diff --git a/Engine/Source/Runtime/Engine/Private/World.cpp b/Engine/Source/Runtime/Engine/Private/World.cpp
--- a/Engine/Source/Runtime/Engine/Private/World.cpp
+++ b/Engine/Source/Runtime/Engine/Private/World.cpp
@@ -320,23 +320,6 @@ FScopedLevelCollectionContextSwitch::~FScopedLevelCollectionContextSwitch()
 	}
 }
 
-FScopedLoadAllExternalObjects::FScopedLoadAllExternalObjects(FName InPackageName)
-{
-	if (!UWorld::LoadAllExternalObjects.Contains(InPackageName))
-	{
-		PackageName = InPackageName;
-		UWorld::LoadAllExternalObjects.Add(PackageName);
-	}
-}
-
-FScopedLoadAllExternalObjects::~FScopedLoadAllExternalObjects()
-{
-	if (PackageName != NAME_None)
-	{
-		UWorld::LoadAllExternalObjects.Remove(PackageName);
-	}
-}
-
 void FWorldPartitionEvents::BroadcastWorldPartitionInitialized(UWorld* InWorld, UWorldPartition* InWorldPartition)
 {
 	check(InWorld);
@@ -362,7 +345,6 @@ FAudioDeviceWorldDelegates::FOnWorldUnregisteredWithAudioDevice FAudioDeviceWorl
 UWorldProxy GWorld;
 
 TMap<FName, EWorldType::Type> UWorld::WorldTypePreLoadMap;
-TSet<FName> UWorld::LoadAllExternalObjects;
 
 FWorldDelegates::FWorldEvent FWorldDelegates::OnPostWorldCreation;
 FWorldDelegates::FWorldInitializationEvent FWorldDelegates::OnPreWorldInitialization;
diff --git a/Engine/Source/Runtime/Engine/Private/WorldPartition/WorldPartitionPackageCache.cpp b/Engine/Source/Runtime/Engine/Private/WorldPartition/WorldPartitionPackageCache.cpp
--- a/Engine/Source/Runtime/Engine/Private/WorldPartition/WorldPartitionPackageCache.cpp
+++ b/Engine/Source/Runtime/Engine/Private/WorldPartition/WorldPartitionPackageCache.cpp
@@ -5,6 +5,7 @@
 #if WITH_EDITOR
 #include "UObject/UObjectHash.h"
 #include "Engine/World.h"
+#include "Engine/Level.h"
 #include "PackageTools.h"
 #include "Misc/PackagePath.h"
 #include "WorldPartition/WorldPartition.h"
@@ -59,8 +60,6 @@ void FWorldPartitionPackageGlobalCache::LoadPackage(FWorldPartitionPackageCache*
 	FLoadPackageAsyncDelegate CompletionCallback = FLoadPackageAsyncDelegate::CreateLambda([this, bInWorldPackage, InClient](const FName& LoadedPackageName, UPackage* LoadedPackage, EAsyncLoadingResult::Type Result)
 	{
 		FPackageCacheInfo& CacheInfo = Cache.FindChecked(LoadedPackageName);
-		check(CacheInfo.ScopedLoadAllExternalObjects.IsValid());
-		CacheInfo.ScopedLoadAllExternalObjects.Reset();
 
 		if (Result == EAsyncLoadingResult::Succeeded)
 		{
@@ -92,15 +91,15 @@ void FWorldPartitionPackageGlobalCache::LoadPackage(FWorldPartitionPackageCache*
 	}
 
 	// Cache doesn't initialize the world, so make sure external objects will be loaded in ULevel::PostLoad
-	CacheInfo.ScopedLoadAllExternalObjects = TUniquePtr<FScopedLoadAllExternalObjects>(new FScopedLoadAllExternalObjects(InPackageName));
+	FLinkerInstancingContext InstancingContext({ ULevel::LoadAllExternalObjectsTag });
 	if (bLoadAsync)
 	{
-		::LoadPackageAsync(FPackagePath::FromPackageNameChecked(InPackageToLoadFrom), InPackageName, CompletionCallback);
+		::LoadPackageAsync(FPackagePath::FromPackageNameChecked(InPackageToLoadFrom), InPackageName, CompletionCallback, PKG_None, INDEX_NONE, 0, &InstancingContext);
 	}
 	else
 	{
 		UPackage* Package = CreatePackage(*InPackageName.ToString());
-		Package = ::LoadPackage(Package, InPackageToLoadFrom, LOAD_None);
+		Package = ::LoadPackage(Package, InPackageToLoadFrom, LOAD_None, nullptr, &InstancingContext);
 		CompletionCallback.Execute(InPackageName, Package, Package ? EAsyncLoadingResult::Succeeded : EAsyncLoadingResult::Failed);
 	}
 }
diff --git a/Engine/Source/Runtime/Engine/Public/WorldPartition/WorldPartitionPackageCache.h b/Engine/Source/Runtime/Engine/Public/WorldPartition/WorldPartitionPackageCache.h
--- a/Engine/Source/Runtime/Engine/Public/WorldPartition/WorldPartitionPackageCache.h
+++ b/Engine/Source/Runtime/Engine/Public/WorldPartition/WorldPartitionPackageCache.h
@@ -8,7 +8,6 @@
 #include "UObject/ObjectMacros.h"
 #include "UObject/UObjectGlobals.h"
 #include "UObject/WeakObjectPtr.h"
-#include "Engine/World.h"
 #include "Misc/Guid.h"
 
 class FLinkerInstancingContext;
@@ -42,7 +41,6 @@ private:
 		TSet<FWorldPartitionPackageCache*> Referencers;
 		TWeakObjectPtr<UPackage> Package;
 		TArray<FLoadPackageAsyncDelegate> CompletionDelegates;
-		TUniquePtr<FScopedLoadAllExternalObjects> ScopedLoadAllExternalObjects;
 	};
 	TMap<FName, FPackageCacheInfo> Cache;
 
