SourceControl: Only show a failure to check-in popup if the operation actually failed (it could have been cancelled).

diff --git a/Engine/Plugins/Editor/ContentBrowser/ContentBrowserAssetDataSource/Source/ContentBrowserAssetDataSource/Private/AssetFileContextMenu.cpp b/Engine/Plugins/Editor/ContentBrowser/ContentBrowserAssetDataSource/Source/ContentBrowserAssetDataSource/Private/AssetFileContextMenu.cpp
--- a/Engine/Plugins/Editor/ContentBrowser/ContentBrowserAssetDataSource/Source/ContentBrowserAssetDataSource/Private/AssetFileContextMenu.cpp
+++ b/Engine/Plugins/Editor/ContentBrowser/ContentBrowserAssetDataSource/Source/ContentBrowserAssetDataSource/Private/AssetFileContextMenu.cpp
@@ -2139,9 +2139,11 @@ void FAssetFileContextMenu::ExecuteSCCCheckIn()
 		GetSelectedPackageNames(PackageNames);
 
 		const bool bUseSourceControlStateCache = true;
-		const bool bCheckinGood = FSourceControlWindows::PromptForCheckin(bUseSourceControlStateCache, PackageNames);
 
-		if (!bCheckinGood)
+		FCheckinResultInfo ResultInfo;
+		FSourceControlWindows::PromptForCheckin(ResultInfo, PackageNames, TArray<FString>(), TArray<FString>(), bUseSourceControlStateCache);
+
+		if ( ResultInfo.Result == ECommandResult::Failed )
 		{
 			FMessageDialog::Open(EAppMsgType::Ok, NSLOCTEXT("UnrealEd", "SCC_Checkin_Failed", "Check-in failed as a result of save failure."));
 		}
