Fixing Content Browser not showing asterisk when asset dirtied after reload.

diff --git a/Engine/Plugins/Editor/ContentBrowser/ContentBrowserAssetDataSource/Source/ContentBrowserAssetDataSource/Private/ContentBrowserAssetDataPayload.cpp b/Engine/Plugins/Editor/ContentBrowser/ContentBrowserAssetDataSource/Source/ContentBrowserAssetDataSource/Private/ContentBrowserAssetDataPayload.cpp
--- a/Engine/Plugins/Editor/ContentBrowser/ContentBrowserAssetDataSource/Source/ContentBrowserAssetDataSource/Private/ContentBrowserAssetDataPayload.cpp
+++ b/Engine/Plugins/Editor/ContentBrowser/ContentBrowserAssetDataSource/Source/ContentBrowserAssetDataSource/Private/ContentBrowserAssetDataPayload.cpp
@@ -30,6 +30,11 @@ FContentBrowserAssetFileItemDataPayload::FContentBrowserAssetFileItemDataPayload
 
 UPackage* FContentBrowserAssetFileItemDataPayload::GetPackage(const bool bTryRecacheIfNull) const
 {
+	if (bHasCachedPackagePtr && CachedPackagePtr.IsStale())
+	{
+		FlushCaches();
+	}
+	
 	if (!bHasCachedPackagePtr || (bTryRecacheIfNull && !CachedPackagePtr.IsValid()))
 	{
 		if (!AssetData.PackageName.IsNone())
@@ -58,6 +63,11 @@ UPackage* FContentBrowserAssetFileItemDataPayload::LoadPackage(TSet<FName> LoadT
 
 UObject* FContentBrowserAssetFileItemDataPayload::GetAsset(const bool bTryRecacheIfNull) const
 {
+	if (bHasCachedAssetPtr && CachedAssetPtr.IsStale())
+	{
+		FlushCaches();
+	}
+	
 	if (!bHasCachedAssetPtr || (bTryRecacheIfNull && !CachedAssetPtr.IsValid()))
 	{
 		if (!AssetData.ObjectPath.IsNone())
@@ -138,6 +148,20 @@ void FContentBrowserAssetFileItemDataPayload::UpdateThumbnail(FAssetThumbnail& I
 	InThumbnail.SetAsset(AssetData);
 }
 
+void FContentBrowserAssetFileItemDataPayload::FlushCaches() const
+{
+	bHasCachedPackagePtr = false;
+	CachedPackagePtr.Reset();
+
+	bHasCachedAssetPtr = false;
+	CachedAssetPtr.Reset();
+
+	bHasCachedAssetTypeActionsPtr = false;
+	CachedAssetTypeActionsPtr.Reset();
+
+	bHasCachedFilename = false;
+	CachedFilename.Reset();
+}
 
 FContentBrowserAssetFileItemDataPayload_Creation::FContentBrowserAssetFileItemDataPayload_Creation(FAssetData&& InAssetData, UClass* InAssetClass, UFactory* InFactory)
 	: FContentBrowserAssetFileItemDataPayload(MoveTemp(InAssetData))
diff --git a/Engine/Plugins/Editor/ContentBrowser/ContentBrowserAssetDataSource/Source/ContentBrowserAssetDataSource/Public/ContentBrowserAssetDataPayload.h b/Engine/Plugins/Editor/ContentBrowser/ContentBrowserAssetDataSource/Source/ContentBrowserAssetDataSource/Public/ContentBrowserAssetDataPayload.h
--- a/Engine/Plugins/Editor/ContentBrowser/ContentBrowserAssetDataSource/Source/ContentBrowserAssetDataSource/Public/ContentBrowserAssetDataPayload.h
+++ b/Engine/Plugins/Editor/ContentBrowser/ContentBrowserAssetDataSource/Source/ContentBrowserAssetDataSource/Public/ContentBrowserAssetDataPayload.h
@@ -75,6 +75,8 @@ private:
 
 	mutable bool bHasCachedFilename = false;
 	mutable FString CachedFilename;
+
+	void FlushCaches() const;
 };
 
 class CONTENTBROWSERASSETDATASOURCE_API FContentBrowserAssetFileItemDataPayload_Creation : public FContentBrowserAssetFileItemDataPayload, public FGCObject
