Treat generated class assets as immutable

#rnx

diff --git a/Engine/Plugins/Editor/ContentBrowser/ContentBrowserAssetDataSource/Source/ContentBrowserAssetDataSource/Private/ContentBrowserAssetDataCore.cpp b/Engine/Plugins/Editor/ContentBrowser/ContentBrowserAssetDataSource/Source/ContentBrowserAssetDataSource/Private/ContentBrowserAssetDataCore.cpp
--- a/Engine/Plugins/Editor/ContentBrowser/ContentBrowserAssetDataSource/Source/ContentBrowserAssetDataSource/Private/ContentBrowserAssetDataCore.cpp
+++ b/Engine/Plugins/Editor/ContentBrowser/ContentBrowserAssetDataSource/Source/ContentBrowserAssetDataSource/Private/ContentBrowserAssetDataCore.cpp
@@ -160,7 +160,27 @@ bool CanModifyAssetFolderItem(IAssetTools* InAssetTools, const FContentBrowserAs
 
 bool CanModifyAssetFileItem(IAssetTools* InAssetTools, const FContentBrowserAssetFileItemDataPayload& InAssetPayload, FText* OutErrorMsg)
 {
-	return CanModifyPath(InAssetTools, InAssetPayload.GetAssetData().PackageName, OutErrorMsg);
+	if (!CanModifyPath(InAssetTools, InAssetPayload.GetAssetData().PackageName, OutErrorMsg))
+	{
+		return false;
+	}
+
+	if (const UClass* AssetClass = InAssetPayload.GetAssetData().GetClass())
+	{
+		if (AssetClass->IsChildOf<UClass>())
+		{
+			SetOptionalErrorMessage(OutErrorMsg, LOCTEXT("Error_CannotModifyGeneratedClasses", "Cannot modify generated classes"));
+			return false;
+		}
+	}
+
+	if (InAssetPayload.GetAssetData().HasAnyPackageFlags(PKG_Cooked | PKG_FilterEditorOnly))
+	{
+		SetOptionalErrorMessage(OutErrorMsg, LOCTEXT("Error_CannotModifyCookedAssets", "Cannot modify cooked assets"));
+		return false;
+	}
+
+	return true;
 }
 
 bool CanEditItem(IAssetTools* InAssetTools, const UContentBrowserDataSource* InOwnerDataSource, const FContentBrowserItemData& InItem, FText* OutErrorMsg)
@@ -180,12 +200,6 @@ bool CanEditAssetFileItem(IAssetTools* InAssetTools, const FContentBrowserAssetF
 		return false;
 	}
 
-	if (InAssetPayload.GetAssetData().PackageFlags & PKG_FilterEditorOnly)
-	{
-		SetOptionalErrorMessage(OutErrorMsg, LOCTEXT("Error_CannotEditCookedPackages", "Cannot edit cooked packages"));
-		return false;
-	}
-
 	return true;
 }
 
@@ -330,6 +344,15 @@ bool CanDuplicateAssetFileItem(IAssetTools* InAssetTools, const FContentBrowserA
 		return false;
 	}
 
+	if (const UClass* AssetClass = InAssetPayload.GetAssetData().GetClass())
+	{
+		if (AssetClass->IsChildOf<UClass>())
+		{
+			SetOptionalErrorMessage(OutErrorMsg, LOCTEXT("Error_CannotDuplicateGeneratedClasses", "Cannot duplicate generated classes"));
+			return false;
+		}
+	}
+
 	if (TSharedPtr<IAssetTypeActions> AssetTypeActions = InAssetPayload.GetAssetTypeActions())
 	{
 		if (!AssetTypeActions->CanDuplicate(InAssetPayload.GetAssetData(), OutErrorMsg))
@@ -437,12 +460,6 @@ bool CanSaveAssetFileItem(IAssetTools* InAssetTools, const FContentBrowserAssetF
 		return false;
 	}
 
-	if (InAssetPayload.GetAssetData().PackageFlags & PKG_FilterEditorOnly)
-	{
-		SetOptionalErrorMessage(OutErrorMsg, LOCTEXT("Error_CannotSaveCookedPackages", "Cannot save cooked packages"));
-		return false;
-	}
-
 	if (EnumHasAnyFlags(InSaveFlags, EContentBrowserItemSaveFlags::SaveOnlyIfLoaded))
 	{
 		// Can't save a package that hasn't been loaded
@@ -692,12 +709,6 @@ bool CanRenameAssetFileItem(IAssetTools* InAssetTools, const FContentBrowserAsse
 		return false;
 	}
 
-	if (InAssetPayload.GetAssetData().PackageFlags & PKG_FilterEditorOnly)
-	{
-		SetOptionalErrorMessage(OutErrorMsg, LOCTEXT("Error_CannotRenameCookedPackages", "Cannot rename cooked packages"));
-		return false;
-	}
-
 	if (InNewName && InNewName->Len() >= NAME_SIZE)
 	{
 		SetOptionalErrorMessage(OutErrorMsg, LOCTEXT("Error_AssetNameTooLarge", "This asset name is too long. Please choose a shorter name."));
diff --git a/Engine/Source/Editor/ContentBrowser/Private/AssetContextMenu.cpp b/Engine/Source/Editor/ContentBrowser/Private/AssetContextMenu.cpp
--- a/Engine/Source/Editor/ContentBrowser/Private/AssetContextMenu.cpp
+++ b/Engine/Source/Editor/ContentBrowser/Private/AssetContextMenu.cpp
@@ -162,15 +162,28 @@ TSharedRef<SWidget> FAssetContextMenu::MakeContextMenu(TArrayView<const FContent
 			ContextObject->CommonClass = CommonClass;
 
 			ContextObject->bCanBeModified = true;
-
-			if (WritableFolderFilter->HasFiltering())
+			for (const UObject* SelectedObject : SelectedObjects)
 			{
-				for (const UObject* SelectedObject : SelectedObjects)
+				if (SelectedObject)
 				{
-					if (SelectedObject)
+					if (const UPackage* SelectedObjectPackage = SelectedObject->GetOutermost())
+					{
+						if (WritableFolderFilter->HasFiltering() && !WritableFolderFilter->PassesStartsWithFilter(SelectedObjectPackage->GetFName()))
+						{
+							ContextObject->bCanBeModified = false;
+							break;
+						}
+
+						if (SelectedObjectPackage->HasAnyPackageFlags(PKG_Cooked | PKG_FilterEditorOnly))
+						{
+							ContextObject->bCanBeModified = false;
+							break;
+						}
+					}
+
+					if (const UClass* AssetClass = SelectedObject->GetClass())
 					{
-						UPackage* SelectedObjectPackage = SelectedObject->GetOutermost();
-						if (SelectedObjectPackage && !WritableFolderFilter->PassesStartsWithFilter(SelectedObjectPackage->GetFName()))
+						if (AssetClass->IsChildOf<UClass>())
 						{
 							ContextObject->bCanBeModified = false;
 							break;
