Add a slow task dialogue when fixing up redirectors

diff --git a/Engine/Plugins/Editor/ContentBrowser/ContentBrowserAssetDataSource/Source/ContentBrowserAssetDataSource/Private/AssetFolderContextMenu.cpp b/Engine/Plugins/Editor/ContentBrowser/ContentBrowserAssetDataSource/Source/ContentBrowserAssetDataSource/Private/AssetFolderContextMenu.cpp
--- a/Engine/Plugins/Editor/ContentBrowser/ContentBrowserAssetDataSource/Source/ContentBrowserAssetDataSource/Private/AssetFolderContextMenu.cpp
+++ b/Engine/Plugins/Editor/ContentBrowser/ContentBrowserAssetDataSource/Source/ContentBrowserAssetDataSource/Private/AssetFolderContextMenu.cpp
@@ -47,6 +47,7 @@
 #include "ContentBrowserDataMenuContexts.h"
 #include "Stats/Stats.h"
 #include "AssetContextMenuUtils.h"
+#include "Misc/ScopedSlowTask.h"
 
 #define LOCTEXT_NAMESPACE "ContentBrowser"
 
@@ -221,41 +222,47 @@ void FAssetFolderContextMenu::ExecuteFixUpRedirectorsInFolder()
 	// Form a filter from the paths
 	FARFilter Filter;
 	Filter.bRecursivePaths = true;
-	for (const auto& Path : SelectedPaths)
+	for (const FString& Path : SelectedPaths)
 	{
 		Filter.PackagePaths.Emplace(*Path);
 		Filter.ClassPaths.Emplace(FTopLevelAssetPath(TEXT("/Script/CoreUObject"), TEXT("ObjectRedirector")));
 	}
-
+	
 	// Query for a list of assets in the selected paths
 	TArray<FAssetData> AssetList;
 	AssetRegistryModule.Get().GetAssets(Filter, AssetList);
 
-	if (AssetList.Num() > 0)
+	if (AssetList.Num() == 0)
 	{
-		TArray<FString> ObjectPaths;
-		for (const auto& Asset : AssetList)
-		{
-			ObjectPaths.Add(Asset.GetObjectPathString());
-		}
+		return;
+	}
+	
+	TArray<FString> ObjectPaths;
+	for (const FAssetData& Asset : AssetList)
+	{
+		ObjectPaths.Add(Asset.GetObjectPathString());
+	}
 
-		TArray<UObject*> Objects;
-		const bool bAllowedToPromptToLoadAssets = true;
-		const bool bLoadRedirects = true;
-		if (AssetViewUtils::LoadAssetsIfNeeded(ObjectPaths, Objects, bAllowedToPromptToLoadAssets, bLoadRedirects))
-		{
-			// Transform Objects array to ObjectRedirectors array
-			TArray<UObjectRedirector*> Redirectors;
-			for (auto Object : Objects)
-			{
-				auto Redirector = CastChecked<UObjectRedirector>(Object);
-				Redirectors.Add(Redirector);
-			}
+	FScopedSlowTask SlowTask(3, LOCTEXT("FixupRedirectorsSlowTask", "Fixing up redirectors"));
+	SlowTask.MakeDialog();
 
-			// Load the asset tools module
-			FAssetToolsModule& AssetToolsModule = FModuleManager::LoadModuleChecked<FAssetToolsModule>(TEXT("AssetTools"));
-			AssetToolsModule.Get().FixupReferencers(Redirectors);
+	TArray<UObject*> Objects;
+	const bool bAllowedToPromptToLoadAssets = true;
+	const bool bLoadRedirects = true;
+	SlowTask.EnterProgressFrame(1, LOCTEXT("FixupRedirectors_LoadAssets", "Loading Assets..."));
+	if (AssetViewUtils::LoadAssetsIfNeeded(ObjectPaths, Objects, bAllowedToPromptToLoadAssets, bLoadRedirects))
+	{
+		// Transform Objects array to ObjectRedirectors array
+		TArray<UObjectRedirector*> Redirectors;
+		for (UObject* Object : Objects)
+		{
+			Redirectors.Add(CastChecked<UObjectRedirector>(Object));
 		}
+
+		SlowTask.EnterProgressFrame(1, LOCTEXT("FixupRedirectors_FixupReferencers", "Fixing up referencers..."));
+		// Load the asset tools module
+		FAssetToolsModule& AssetToolsModule = FModuleManager::LoadModuleChecked<FAssetToolsModule>(TEXT("AssetTools"));
+		AssetToolsModule.Get().FixupReferencers(Redirectors);
 	}
 }
 
