Fix display of animation layer interfaces on content browser tiles

Use GetDisplayNameFromAssetData (as with thumbnails) rather than just GetName

diff --git a/Engine/Plugins/Editor/ContentBrowser/ContentBrowserAssetDataSource/Source/ContentBrowserAssetDataSource/Private/ContentBrowserAssetDataCore.cpp b/Engine/Plugins/Editor/ContentBrowser/ContentBrowserAssetDataSource/Source/ContentBrowserAssetDataSource/Private/ContentBrowserAssetDataCore.cpp
--- a/Engine/Plugins/Editor/ContentBrowser/ContentBrowserAssetDataSource/Source/ContentBrowserAssetDataSource/Private/ContentBrowserAssetDataCore.cpp
+++ b/Engine/Plugins/Editor/ContentBrowser/ContentBrowserAssetDataSource/Source/ContentBrowserAssetDataSource/Private/ContentBrowserAssetDataCore.cpp
@@ -1352,7 +1352,15 @@ bool GetAssetFileItemAttribute(const FContentBrowserAssetFileItemDataPayload& In
 		{
 			if (TSharedPtr<IAssetTypeActions> AssetTypeActions = InAssetPayload.GetAssetTypeActions())
 			{
-				OutAttributeValue.SetValue(AssetTypeActions->GetName());
+				const FText AssetDisplayName = AssetTypeActions->GetDisplayNameFromAssetData(InAssetPayload.GetAssetData());
+				if (!AssetDisplayName.IsEmpty())
+				{
+					OutAttributeValue.SetValue(AssetDisplayName);
+				}
+				else
+				{
+					OutAttributeValue.SetValue(AssetTypeActions->GetName());
+				}
 				return true;
 			}
 			return false;
diff --git a/Engine/Plugins/Editor/ContentBrowser/ContentBrowserClassDataSource/Source/ContentBrowserClassDataSource/Private/ContentBrowserClassDataCore.cpp b/Engine/Plugins/Editor/ContentBrowser/ContentBrowserClassDataSource/Source/ContentBrowserClassDataSource/Private/ContentBrowserClassDataCore.cpp
--- a/Engine/Plugins/Editor/ContentBrowser/ContentBrowserClassDataSource/Source/ContentBrowserClassDataSource/Private/ContentBrowserClassDataCore.cpp
+++ b/Engine/Plugins/Editor/ContentBrowser/ContentBrowserClassDataSource/Source/ContentBrowserClassDataSource/Private/ContentBrowserClassDataCore.cpp
@@ -469,7 +469,16 @@ bool GetClassFileItemAttribute(IAssetTypeActions* InClassTypeActions, const FCon
 
 		if (InAttributeKey == ContentBrowserItemAttributes::ItemTypeDisplayName)
 		{
-			OutAttributeValue.SetValue(InClassTypeActions->GetName());
+			const FText AssetDisplayName = InClassTypeActions->GetDisplayNameFromAssetData(InClassPayload.GetAssetData());
+			if (!AssetDisplayName.IsEmpty())
+			{
+				OutAttributeValue.SetValue(AssetDisplayName);
+			}
+			else
+			{
+				OutAttributeValue.SetValue(InClassTypeActions->GetName());
+			}
+
 			return true;
 		}
 
diff --git a/Engine/Source/Developer/AssetTools/Private/AssetTypeActions/AssetTypeActions_AnimBlueprint.cpp b/Engine/Source/Developer/AssetTools/Private/AssetTypeActions/AssetTypeActions_AnimBlueprint.cpp
--- a/Engine/Source/Developer/AssetTools/Private/AssetTypeActions/AssetTypeActions_AnimBlueprint.cpp
+++ b/Engine/Source/Developer/AssetTools/Private/AssetTypeActions/AssetTypeActions_AnimBlueprint.cpp
@@ -358,4 +358,20 @@ TSharedPtr<SWidget> FAssetTypeActions_AnimBlueprint::GetThumbnailOverlay(const F
 		];
 }
 
+FText FAssetTypeActions_AnimBlueprint::GetDisplayNameFromAssetData(const FAssetData& AssetData) const
+{
+	FString BlueprintTypeValue;
+	AssetData.GetTagValue("BlueprintType", BlueprintTypeValue);
+	if(BlueprintTypeValue.Equals(TEXT("BPTYPE_Normal")))
+	{
+		return LOCTEXT("AssetTypeActions_AnimBlueprint", "Animation Blueprint");
+	}
+	else if(BlueprintTypeValue.Equals(TEXT("BPTYPE_Interface")))
+	{
+		return LOCTEXT("AssetTypeActions_AnimLayerInterface", "Animation Layer Interface");
+	}
+
+	return GetName();
+}
+
 #undef LOCTEXT_NAMESPACE
diff --git a/Engine/Source/Developer/AssetTools/Private/AssetTypeActions/AssetTypeActions_AnimBlueprint.h b/Engine/Source/Developer/AssetTools/Private/AssetTypeActions/AssetTypeActions_AnimBlueprint.h
--- a/Engine/Source/Developer/AssetTools/Private/AssetTypeActions/AssetTypeActions_AnimBlueprint.h
+++ b/Engine/Source/Developer/AssetTools/Private/AssetTypeActions/AssetTypeActions_AnimBlueprint.h
@@ -23,7 +23,8 @@ public:
 	virtual void PerformAssetDiff(UObject* Asset1, UObject* Asset2, const struct FRevisionInfo& OldRevision, const struct FRevisionInfo& NewRevision) const override;
 	virtual class UThumbnailInfo* GetThumbnailInfo(UObject* Asset) const override;
 	virtual TSharedPtr<SWidget> GetThumbnailOverlay(const FAssetData& AssetData) const override;
-
+	virtual FText GetDisplayNameFromAssetData(const FAssetData& AssetData) const override;
+	
 	// FAssetTypeActions_Blueprint interface
 	virtual UFactory* GetFactoryForBlueprintType(UBlueprint* InBlueprint) const override;
 
