Remove / disable 'Check-in' context menu item in Content Explorer and Scene Outliner for Source Control providers that do not use changelists

diff --git a/Engine/Plugins/Editor/ContentBrowser/ContentBrowserAssetDataSource/Source/ContentBrowserAssetDataSource/Private/AssetFileContextMenu.cpp b/Engine/Plugins/Editor/ContentBrowser/ContentBrowserAssetDataSource/Source/ContentBrowserAssetDataSource/Private/AssetFileContextMenu.cpp
--- a/Engine/Plugins/Editor/ContentBrowser/ContentBrowserAssetDataSource/Source/ContentBrowserAssetDataSource/Private/AssetFileContextMenu.cpp
+++ b/Engine/Plugins/Editor/ContentBrowser/ContentBrowserAssetDataSource/Source/ContentBrowserAssetDataSource/Private/AssetFileContextMenu.cpp
@@ -2346,7 +2346,15 @@ bool FAssetFileContextMenu::CanExecuteSCCOpenForAdd() const
 
 bool FAssetFileContextMenu::CanExecuteSCCCheckIn() const
 {
-	return bCanExecuteSCCCheckIn;
+	bool bUsesFileRevisions = true;
+
+	ISourceControlProvider& SourceControlProvider = ISourceControlModule::Get().GetProvider();
+	if (ISourceControlModule::Get().IsEnabled())
+	{
+		bUsesFileRevisions = SourceControlProvider.UsesFileRevisions();
+	}
+	
+	return bCanExecuteSCCCheckIn && bUsesFileRevisions;
 }
 
 bool FAssetFileContextMenu::CanExecuteSCCHistory() const
diff --git a/Engine/Plugins/Editor/ContentBrowser/ContentBrowserAssetDataSource/Source/ContentBrowserAssetDataSource/Private/AssetFolderContextMenu.cpp b/Engine/Plugins/Editor/ContentBrowser/ContentBrowserAssetDataSource/Source/ContentBrowserAssetDataSource/Private/AssetFolderContextMenu.cpp
--- a/Engine/Plugins/Editor/ContentBrowser/ContentBrowserAssetDataSource/Source/ContentBrowserAssetDataSource/Private/AssetFolderContextMenu.cpp
+++ b/Engine/Plugins/Editor/ContentBrowser/ContentBrowserAssetDataSource/Source/ContentBrowserAssetDataSource/Private/AssetFolderContextMenu.cpp
@@ -417,7 +417,15 @@ bool FAssetFolderContextMenu::CanExecuteSCCOpenForAdd() const
 
 bool FAssetFolderContextMenu::CanExecuteSCCCheckIn() const
 {
-	return bCanExecuteSCCCheckIn && SelectedPaths.Num() > 0;
+	bool bUsesFileRevisions = true;
+
+	ISourceControlProvider& SourceControlProvider = ISourceControlModule::Get().GetProvider();
+	if (ISourceControlModule::Get().IsEnabled())
+	{
+		bUsesFileRevisions = SourceControlProvider.UsesFileRevisions();
+	}
+
+	return bCanExecuteSCCCheckIn && SelectedPaths.Num() > 0 && bUsesFileRevisions;
 }
 
 bool FAssetFolderContextMenu::CanExecuteSCCSync() const
diff --git a/Engine/Source/Editor/SceneOutliner/Private/SceneOutlinerSCCHandler.cpp b/Engine/Source/Editor/SceneOutliner/Private/SceneOutlinerSCCHandler.cpp
--- a/Engine/Source/Editor/SceneOutliner/Private/SceneOutlinerSCCHandler.cpp
+++ b/Engine/Source/Editor/SceneOutliner/Private/SceneOutlinerSCCHandler.cpp
@@ -142,7 +142,15 @@ bool FSceneOutlinerSCCHandler::CanExecuteSCCCheckOut() const
 
 bool FSceneOutlinerSCCHandler::CanExecuteSCCCheckIn() const
 {
-	return bCanExecuteSCCCheckIn;
+	bool bUsesFileRevisions = true;
+
+	ISourceControlProvider& SourceControlProvider = ISourceControlModule::Get().GetProvider();
+	if (ISourceControlModule::Get().IsEnabled())
+	{
+		bUsesFileRevisions = SourceControlProvider.UsesFileRevisions();
+	}
+
+	return bCanExecuteSCCCheckIn && bUsesFileRevisions;
 }
 
 bool FSceneOutlinerSCCHandler::CanExecuteSCCHistory() const
