Fixed a couple of operations in the asset view contextual menu that would only work on the selected folders when the selection contains both asset items and folder items.

This don't resolve all issue with these operations since some of them should also work with items that are not created by the ContentBrowserAssetDataSource, but this is still better then what we currently have

Here are the fixed operations.
Fix Up Redirectors
Migrate...
Source Control -> Check out
Source Control -> Mark For Add
Source Control -> Check In
Source Control -> Sync

diff --git a/Engine/Plugins/Editor/ContentBrowser/ContentBrowserAssetDataSource/Source/ContentBrowserAssetDataSource/Private/AssetFolderContextMenu.cpp b/Engine/Plugins/Editor/ContentBrowser/ContentBrowserAssetDataSource/Source/ContentBrowserAssetDataSource/Private/AssetFolderContextMenu.cpp
--- a/Engine/Plugins/Editor/ContentBrowser/ContentBrowserAssetDataSource/Source/ContentBrowserAssetDataSource/Private/AssetFolderContextMenu.cpp
+++ b/Engine/Plugins/Editor/ContentBrowser/ContentBrowserAssetDataSource/Source/ContentBrowserAssetDataSource/Private/AssetFolderContextMenu.cpp
@@ -52,8 +52,14 @@
 #define LOCTEXT_NAMESPACE "ContentBrowser"
 
 void FAssetFolderContextMenu::MakeContextMenu(UToolMenu* InMenu, const TArray<FString>& InSelectedPackagePaths)
+{
+	MakeContextMenu(InMenu, InSelectedPackagePaths, TArray<FString>());
+}
+
+void FAssetFolderContextMenu::MakeContextMenu(UToolMenu* InMenu, const TArray<FString>& InSelectedPackagePaths, const TArray<FString>& InSelectedPackages)
 {
 	SelectedPaths = InSelectedPackagePaths;
+	SelectedPackages = InSelectedPackages;
 
 	if (SelectedPaths.Num() > 0)
 	{
@@ -88,21 +94,21 @@ void FAssetFolderContextMenu::AddMenuOptions(UToolMenu* Menu)
 			{
 				FToolMenuEntry& Entry = Section.AddMenuEntry(
 					"FixUpRedirectorsInFolder",
-					LOCTEXT("FixUpRedirectorsInFolder", "Fix Up Redirectors in Folder"),
-					LOCTEXT("FixUpRedirectorsInFolderTooltip", "Finds referencers to all redirectors in the selected folders and resaves them if possible, then deletes any redirectors that had all their referencers fixed."),
+					LOCTEXT("FixUpRedirectorsInFolder", "Fix Up Redirectors"),
+					LOCTEXT("FixUpRedirectorsInFolderTooltip", "Finds referencers to all redirectors in the selected items and resaves them if possible, then deletes any redirectors that had all their referencers fixed."),
 					FSlateIcon(FAppStyle::GetAppStyleSetName(), "Icons.Adjust"),
 					FUIAction(FExecuteAction::CreateSP(this, &FAssetFolderContextMenu::ExecuteFixUpRedirectorsInFolder))
 					);
 				Entry.InsertPosition = FToolMenuInsert("Delete", EToolMenuInsertType::After);
 			}
 
-			if (SelectedPaths.Num() > 0)
+			if (!SelectedPaths.IsEmpty() || !SelectedPackages.IsEmpty())
 			{
 				// Migrate Folder
 				FToolMenuEntry& Entry = Section.AddMenuEntry(
 					"MigrateFolder",
 					LOCTEXT("MigrateFolder", "Migrate..."),
-					LOCTEXT("MigrateFolderTooltip", "Copies assets found in this folder and their dependencies to another game content folder."),
+					LOCTEXT("MigrateFolderTooltip", "Copies assets found in the selection and their dependencies to another content folder."),
 					FSlateIcon(FAppStyle::GetAppStyleSetName(), "ContentBrowser.Migrate"),
 					FUIAction( FExecuteAction::CreateSP( this, &FAssetFolderContextMenu::ExecuteMigrateFolder ) )
 					);
@@ -124,7 +130,7 @@ void FAssetFolderContextMenu::AddMenuOptions(UToolMenu* Menu)
 					Section,
 					"FolderSCCCheckOut",
 					LOCTEXT("FolderSCCCheckOut", "Check Out"),
-					LOCTEXT("FolderSCCCheckOutTooltip", "Checks out all assets from source control which are in this folder."),
+					LOCTEXT("FolderSCCCheckOutTooltip", "Checks out all assets from source control which are in this selection."),
 					FUIAction(
 						FExecuteAction::CreateSP(this, &FAssetFolderContextMenu::ExecuteSCCCheckOut),
 						FCanExecuteAction::CreateSP(this, &FAssetFolderContextMenu::CanExecuteSCCCheckOut)
@@ -137,7 +143,7 @@ void FAssetFolderContextMenu::AddMenuOptions(UToolMenu* Menu)
 					Section,
 					"FolderSCCOpenForAdd",
 					LOCTEXT("FolderSCCOpenForAdd", "Mark For Add"),
-					LOCTEXT("FolderSCCOpenForAddTooltip", "Adds all assets to source control that are in this folder and not already added."),
+					LOCTEXT("FolderSCCOpenForAddTooltip", "Adds all assets to source control that are in this selection and not already added."),
 					FUIAction(
 						FExecuteAction::CreateSP(this, &FAssetFolderContextMenu::ExecuteSCCOpenForAdd),
 						FCanExecuteAction::CreateSP(this, &FAssetFolderContextMenu::CanExecuteSCCOpenForAdd)
@@ -150,7 +156,7 @@ void FAssetFolderContextMenu::AddMenuOptions(UToolMenu* Menu)
 					Section,
 					"FolderSCCCheckIn",
 					LOCTEXT("FolderSCCCheckIn", "Check In"),
-					LOCTEXT("FolderSCCCheckInTooltip", "Checks in all assets to source control which are in this folder."),
+					LOCTEXT("FolderSCCCheckInTooltip", "Checks in all assets to source control which are in this selection."),
 					FUIAction(
 						FExecuteAction::CreateSP(this, &FAssetFolderContextMenu::ExecuteSCCCheckIn),
 						FCanExecuteAction::CreateSP(this, &FAssetFolderContextMenu::CanExecuteSCCCheckIn)
@@ -162,7 +168,7 @@ void FAssetFolderContextMenu::AddMenuOptions(UToolMenu* Menu)
 				Section.AddMenuEntry(
 					"FolderSCCSync",
 					LOCTEXT("FolderSCCSync", "Sync"),
-					LOCTEXT("FolderSCCSyncTooltip", "Syncs all the assets in this folder to the latest version."),
+					LOCTEXT("FolderSCCSyncTooltip", "Syncs all the assets in this selection to the latest version."),
 					FSlateIcon(),
 					FUIAction(
 						FExecuteAction::CreateSP(this, &FAssetFolderContextMenu::ExecuteSCCSync),
@@ -209,6 +215,8 @@ void FAssetFolderContextMenu::ExecuteMigrateFolder()
 			PackageNames.Add((*AssetIt).PackageName);
 		}
 
+		PackageNames.Append(SelectedPackages);
+
 		// Load all the assets in the selected paths
 		FAssetToolsModule& AssetToolsModule = FModuleManager::Get().LoadModuleChecked<FAssetToolsModule>("AssetTools");
 		AssetToolsModule.Get().MigratePackages( PackageNames );
@@ -222,16 +230,37 @@ void FAssetFolderContextMenu::ExecuteFixUpRedirectorsInFolder()
 	// Form a filter from the paths
 	FARFilter Filter;
 	Filter.bRecursivePaths = true;
+
+	Filter.PackagePaths.Reserve(SelectedPaths.Num());
 	for (const FString& Path : SelectedPaths)
 	{
 		Filter.PackagePaths.Emplace(*Path);
-		Filter.ClassPaths.Emplace(FTopLevelAssetPath(TEXT("/Script/CoreUObject"), TEXT("ObjectRedirector")));
+	}
+
+	if (!SelectedPaths.IsEmpty())
+	{
+		Filter.ClassPaths.Add(UObjectRedirector::StaticClass()->GetClassPathName());
 	}
 	
 	// Query for a list of assets in the selected paths
 	TArray<FAssetData> AssetList;
 	AssetRegistryModule.Get().GetAssets(Filter, AssetList);
 
+	Filter.Clear();
+
+	Filter.PackageNames.Reserve(SelectedPackages.Num());
+	for (const FString& PackageName : SelectedPackages)
+	{
+		Filter.PackageNames.Emplace(*PackageName);
+	}
+
+	if (!SelectedPackages.IsEmpty())
+	{
+		Filter.ClassPaths.Add(UObjectRedirector::StaticClass()->GetClassPathName());
+	}
+
+	AssetRegistryModule.Get().GetAssets(Filter, AssetList);
+
 	if (AssetList.Num() == 0)
 	{
 		return;
@@ -271,6 +300,7 @@ void FAssetFolderContextMenu::ExecuteSCCCheckOut()
 	// Get a list of package names in the selected paths
 	TArray<FString> PackageNames;
 	GetPackageNamesInSelectedPaths(PackageNames);
+	PackageNames.Append(SelectedPackages);
 
 	TArray<UPackage*> PackagesToCheckOut;
 	for ( auto PackageIt = PackageNames.CreateConstIterator(); PackageIt; ++PackageIt )
@@ -301,6 +331,7 @@ void FAssetFolderContextMenu::ExecuteSCCOpenForAdd()
 	// Get a list of package names in the selected paths
 	TArray<FString> PackageNames;
 	GetPackageNamesInSelectedPaths(PackageNames);
+	PackageNames.Append(SelectedPackages);
 
 	TArray<FString> PackagesToAdd;
 	TArray<UPackage*> PackagesToSave;
@@ -355,6 +386,7 @@ void FAssetFolderContextMenu::ExecuteSCCCheckIn()
 	// Get a list of package names in the selected paths
 	TArray<FString> PackageNames;
 	GetPackageNamesInSelectedPaths(PackageNames);
+	PackageNames.Append(SelectedPackages);
 
 	// Form a list of loaded packages to prompt for save
 	TArray<UPackage*> LoadedPackages;
@@ -398,6 +430,7 @@ void FAssetFolderContextMenu::ExecuteSCCCheckIn()
 void FAssetFolderContextMenu::ExecuteSCCSync() const
 {
 	AssetViewUtils::SyncPathsFromSourceControl(SelectedPaths);
+	AssetViewUtils::SyncPackagesFromSourceControl(SelectedPackages);
 }
 
 void FAssetFolderContextMenu::ExecuteSCCConnect() const
@@ -407,12 +440,12 @@ void FAssetFolderContextMenu::ExecuteSCCConnect() const
 
 bool FAssetFolderContextMenu::CanExecuteSCCCheckOut() const
 {
-	return bCanExecuteSCCCheckOut && SelectedPaths.Num() > 0;
+	return bCanExecuteSCCCheckOut && (!SelectedPaths.IsEmpty() || !SelectedPackages.IsEmpty());
 }
 
 bool FAssetFolderContextMenu::CanExecuteSCCOpenForAdd() const
 {
-	return bCanExecuteSCCOpenForAdd && SelectedPaths.Num() > 0;
+	return bCanExecuteSCCOpenForAdd && (!SelectedPaths.IsEmpty() || !SelectedPackages.IsEmpty());
 }
 
 bool FAssetFolderContextMenu::CanExecuteSCCCheckIn() const
@@ -425,12 +458,12 @@ bool FAssetFolderContextMenu::CanExecuteSCCCheckIn() const
 		bUsesFileRevisions = SourceControlProvider.UsesFileRevisions();
 	}
 
-	return bCanExecuteSCCCheckIn && SelectedPaths.Num() > 0 && bUsesFileRevisions;
+	return bCanExecuteSCCCheckIn && (!SelectedPaths.IsEmpty() || !SelectedPackages.IsEmpty()) && bUsesFileRevisions;
 }
 
 bool FAssetFolderContextMenu::CanExecuteSCCSync() const
 {
-	if (SelectedPaths.Num() > 0)
+	if (!SelectedPaths.IsEmpty() || !SelectedPackages.IsEmpty())
 	{
 		ISourceControlProvider& SourceControlProvider = ISourceControlModule::Get().GetProvider();
 		if (SourceControlProvider.IsEnabled() && SourceControlProvider.UsesFileRevisions())
@@ -475,6 +508,7 @@ void FAssetFolderContextMenu::StartProcessCanExecuteVars()
 	if ( SourceControlProvider.IsEnabled() && SourceControlProvider.IsAvailable() )
 	{
 		GetPackageNamesInSelectedPaths(PackageNamesToProcess);
+		PackageNamesToProcess.Append(SelectedPackages);
 		CurrentPackageIndex = 0;
 		ProcessCanExecuteVars();
 	}
diff --git a/Engine/Plugins/Editor/ContentBrowser/ContentBrowserAssetDataSource/Source/ContentBrowserAssetDataSource/Private/ContentBrowserAssetDataCore.cpp b/Engine/Plugins/Editor/ContentBrowser/ContentBrowserAssetDataSource/Source/ContentBrowserAssetDataSource/Private/ContentBrowserAssetDataCore.cpp
--- a/Engine/Plugins/Editor/ContentBrowser/ContentBrowserAssetDataSource/Source/ContentBrowserAssetDataSource/Private/ContentBrowserAssetDataCore.cpp
+++ b/Engine/Plugins/Editor/ContentBrowser/ContentBrowserAssetDataSource/Source/ContentBrowserAssetDataSource/Private/ContentBrowserAssetDataCore.cpp
@@ -1544,6 +1544,7 @@ void PopulateAssetFolderContextMenu(UContentBrowserDataSource* InOwnerDataSource
 
 	// Extract the internal package paths that belong to this data source from the full list of selected items given in the context
 	TArray<FString> SelectedPackagePaths;
+	TArray<FString> SelectedAssetPackages;
 	for (const FContentBrowserItem& SelectedItem : ContextObject->SelectedItems)
 	{
 		for (const FContentBrowserItemData& SelectedItemData : SelectedItem.GetInternalItems())
@@ -1552,12 +1553,17 @@ void PopulateAssetFolderContextMenu(UContentBrowserDataSource* InOwnerDataSource
 			{
 				SelectedPackagePaths.Add(FolderPayload->GetInternalPath().ToString());
 			}
+			else if (TSharedPtr<const FContentBrowserAssetFileItemDataPayload> ItemPayload = GetAssetFileItemPayload(InOwnerDataSource, SelectedItemData))
+			{
+				SelectedAssetPackages.Add(ItemPayload->GetAssetData().PackageName.ToString());
+			}
 		}
 	}
 
 	InAssetFolderContextMenu.MakeContextMenu(
 		InMenu,
-		SelectedPackagePaths
+		SelectedPackagePaths,
+		SelectedAssetPackages
 		);
 }
 
diff --git a/Engine/Plugins/Editor/ContentBrowser/ContentBrowserAssetDataSource/Source/ContentBrowserAssetDataSource/Public/AssetFolderContextMenu.h b/Engine/Plugins/Editor/ContentBrowser/ContentBrowserAssetDataSource/Source/ContentBrowserAssetDataSource/Public/AssetFolderContextMenu.h
--- a/Engine/Plugins/Editor/ContentBrowser/ContentBrowserAssetDataSource/Source/ContentBrowserAssetDataSource/Public/AssetFolderContextMenu.h
+++ b/Engine/Plugins/Editor/ContentBrowser/ContentBrowserAssetDataSource/Source/ContentBrowserAssetDataSource/Public/AssetFolderContextMenu.h
@@ -17,10 +17,22 @@ class CONTENTBROWSERASSETDATASOURCE_API FAssetFolderContextMenu : public TShared
 public:
 	virtual ~FAssetFolderContextMenu() = default;
 
-	/** Makes the context menu widget */
+	UE_DEPRECATED(5.1, "Use the new overload.")
 	void MakeContextMenu(
 		UToolMenu* InMenu,
 		const TArray<FString>& InSelectedPackagePaths
+	);
+
+	/** 
+	 * Makes the context menu widget
+	 * @param InMenu The menu in which the entries will be added
+	 * @param InSelectedPackagePaths The path of the selected folders
+	 * @param InSelectedPackages The package names of the selected asset item. Empty if the selection is only some folder items.
+	 */
+	void MakeContextMenu(
+		UToolMenu* InMenu,
+		const TArray<FString>& InSelectedPackagePaths,
+		const TArray<FString>& InSelectedPackages
 		);
 
 private:
@@ -91,6 +103,10 @@ private:
 	void StopProcessCanExecuteVars();
 
 	TArray<FString> SelectedPaths;
+
+	// Not null if the selection contained some assets also
+	TArray<FString> SelectedPackages;
+
 	TWeakPtr<SWidget> ParentWidget;
 
 	/** Cached SCC CanExecute vars */
diff --git a/Engine/Source/Editor/ContentBrowser/Private/SContentBrowser.cpp b/Engine/Source/Editor/ContentBrowser/Private/SContentBrowser.cpp
--- a/Engine/Source/Editor/ContentBrowser/Private/SContentBrowser.cpp
+++ b/Engine/Source/Editor/ContentBrowser/Private/SContentBrowser.cpp
@@ -4075,7 +4075,8 @@ TSharedPtr<SWidget> SContentBrowser::GetItemContextMenu(TArrayView<const FConten
 
 		{
 			UContentBrowserDataMenuContext_FolderMenu* DataContextObject = NewObject<UContentBrowserDataMenuContext_FolderMenu>();
-			DataContextObject->SelectedItems = PathContextMenu->GetSelectedFolders();
+			// Include the items that are not folders to help the batch operations operate on these also.
+			DataContextObject->SelectedItems = SelectedItems;
 			DataContextObject->bCanBeModified = Context->bCanBeModified;
 			DataContextObject->ParentWidget = ViewContext == EContentBrowserViewContext::AssetView ? TSharedPtr<SWidget>(AssetViewPtr) : ViewContext == EContentBrowserViewContext::FavoriteView ? TSharedPtr<SWidget>(FavoritePathViewPtr) : TSharedPtr<SWidget>(PathViewPtr);
 			MenuContext.AddObject(DataContextObject);
