Hide browse to header menu when C++ is not available

diff --git a/Engine/Plugins/Editor/ContentBrowser/ContentBrowserAssetDataSource/Source/ContentBrowserAssetDataSource/Private/AssetFileContextMenu.cpp b/Engine/Plugins/Editor/ContentBrowser/ContentBrowserAssetDataSource/Source/ContentBrowserAssetDataSource/Private/AssetFileContextMenu.cpp
--- a/Engine/Plugins/Editor/ContentBrowser/ContentBrowserAssetDataSource/Source/ContentBrowserAssetDataSource/Private/AssetFileContextMenu.cpp
+++ b/Engine/Plugins/Editor/ContentBrowser/ContentBrowserAssetDataSource/Source/ContentBrowserAssetDataSource/Private/AssetFileContextMenu.cpp
@@ -82,6 +82,10 @@
 
 #include "ContentBrowserDataMenuContexts.h"
 
+#include "Editor/UnrealEdEngine.h"
+#include "Preferences/UnrealEdOptions.h"
+#include "UnrealEdGlobals.h"
+
 #define LOCTEXT_NAMESPACE "ContentBrowser"
 
 void FAssetFileContextMenu::MakeContextMenu(UToolMenu* InMenu, const TArray<FAssetData>& InSelectedAssets, const FOnShowAssetsInPathsView& InOnShowAssetsInPathsView)
@@ -1054,7 +1058,7 @@ bool FAssetFileContextMenu::AddDocumentationMenuOptions(UToolMenu* Menu)
 			}
 		}
 
-		if ( !bIsBlueprint && FSourceCodeNavigation::IsCompilerAvailable() )
+		if (!bIsBlueprint && FSourceCodeNavigation::IsCompilerAvailable() && ensure(GUnrealEd) && GUnrealEd->GetUnrealEdOptions()->IsCPPAllowed())
 		{
 			FString ClassHeaderPath;
 			if( FSourceCodeNavigation::FindClassHeaderPath( SelectedClass, ClassHeaderPath ) && IFileManager::Get().FileSize( *ClassHeaderPath ) != INDEX_NONE )
