Refactored Asset File Context Menu and Asset Folder Context Menu dynamic items to use caching

diff --git a/Engine/Plugins/Editor/ContentBrowser/ContentBrowserAssetDataSource/Source/ContentBrowserAssetDataSource/Private/AssetFileContextMenu.cpp b/Engine/Plugins/Editor/ContentBrowser/ContentBrowserAssetDataSource/Source/ContentBrowserAssetDataSource/Private/AssetFileContextMenu.cpp
--- a/Engine/Plugins/Editor/ContentBrowser/ContentBrowserAssetDataSource/Source/ContentBrowserAssetDataSource/Private/AssetFileContextMenu.cpp
+++ b/Engine/Plugins/Editor/ContentBrowser/ContentBrowserAssetDataSource/Source/ContentBrowserAssetDataSource/Private/AssetFileContextMenu.cpp
@@ -2344,15 +2344,7 @@ bool FAssetFileContextMenu::CanExecuteSCCOpenForAdd() const
 }
 
 bool FAssetFileContextMenu::CanExecuteSCCCheckIn() const
-{
-	bool bUsesFileRevisions = true;
-
-	ISourceControlProvider& SourceControlProvider = ISourceControlModule::Get().GetProvider();
-	if (ISourceControlModule::Get().IsEnabled())
-	{
-		bUsesFileRevisions = SourceControlProvider.UsesFileRevisions();
-	}
-	
+{	
 	return bCanExecuteSCCCheckIn && bUsesFileRevisions;
 }
 
@@ -2418,6 +2410,7 @@ void FAssetFileContextMenu::CacheCanExecuteVars()
 	bCanExecuteSCCHistory = false;
 	bCanExecuteSCCRevert = false;
 	bCanExecuteSCCSync = false;
+	bUsesFileRevisions = false;
 
 	for (auto AssetIt = SelectedAssets.CreateConstIterator(); AssetIt; ++AssetIt)
 	{
@@ -2489,6 +2482,11 @@ void FAssetFileContextMenu::CacheCanExecuteVars()
 			break;
 		}
 	}
+
+	if (ISourceControlModule::Get().IsEnabled() && ISourceControlModule::Get().GetProvider().IsAvailable())
+	{
+		bUsesFileRevisions = ISourceControlModule::Get().GetProvider().UsesChangelists();
+	}
 }
 
 void FAssetFileContextMenu::GetSelectedPackageNames(TArray<FString>& OutPackageNames) const
diff --git a/Engine/Plugins/Editor/ContentBrowser/ContentBrowserAssetDataSource/Source/ContentBrowserAssetDataSource/Private/AssetFolderContextMenu.cpp b/Engine/Plugins/Editor/ContentBrowser/ContentBrowserAssetDataSource/Source/ContentBrowserAssetDataSource/Private/AssetFolderContextMenu.cpp
--- a/Engine/Plugins/Editor/ContentBrowser/ContentBrowserAssetDataSource/Source/ContentBrowserAssetDataSource/Private/AssetFolderContextMenu.cpp
+++ b/Engine/Plugins/Editor/ContentBrowser/ContentBrowserAssetDataSource/Source/ContentBrowserAssetDataSource/Private/AssetFolderContextMenu.cpp
@@ -450,14 +450,6 @@ bool FAssetFolderContextMenu::CanExecuteSCCOpenForAdd() const
 
 bool FAssetFolderContextMenu::CanExecuteSCCCheckIn() const
 {
-	bool bUsesFileRevisions = true;
-
-	ISourceControlProvider& SourceControlProvider = ISourceControlModule::Get().GetProvider();
-	if (ISourceControlModule::Get().IsEnabled())
-	{
-		bUsesFileRevisions = SourceControlProvider.UsesFileRevisions();
-	}
-
 	return bCanExecuteSCCCheckIn && (!SelectedPaths.IsEmpty() || !SelectedPackages.IsEmpty()) && bUsesFileRevisions;
 }
 
@@ -503,6 +495,7 @@ void FAssetFolderContextMenu::StartProcessCanExecuteVars()
 	bCanExecuteSCCCheckOut = false;
 	bCanExecuteSCCOpenForAdd = false;
 	bCanExecuteSCCCheckIn = false;
+	bUsesFileRevisions = false;
 
 	ISourceControlProvider& SourceControlProvider = ISourceControlModule::Get().GetProvider();
 	if ( SourceControlProvider.IsEnabled() && SourceControlProvider.IsAvailable() )
@@ -602,6 +595,8 @@ void FAssetFolderContextMenu::ProcessCanExecuteVars()
 		// Everything was processed stop ticking
 		StopProcessCanExecuteVars();
 	}
+
+	bUsesFileRevisions = SourceControlProvider.UsesFileRevisions();
 }
 
 void FAssetFolderContextMenu::StopProcessCanExecuteVars()
diff --git a/Engine/Plugins/Editor/ContentBrowser/ContentBrowserAssetDataSource/Source/ContentBrowserAssetDataSource/Public/AssetFileContextMenu.h b/Engine/Plugins/Editor/ContentBrowser/ContentBrowserAssetDataSource/Source/ContentBrowserAssetDataSource/Public/AssetFileContextMenu.h
--- a/Engine/Plugins/Editor/ContentBrowser/ContentBrowserAssetDataSource/Source/ContentBrowserAssetDataSource/Public/AssetFileContextMenu.h
+++ b/Engine/Plugins/Editor/ContentBrowser/ContentBrowserAssetDataSource/Source/ContentBrowserAssetDataSource/Public/AssetFileContextMenu.h
@@ -319,6 +319,7 @@ private:
 	bool bCanExecuteSCCHistory = false;
 	bool bCanExecuteSCCRevert = false;
 	bool bCanExecuteSCCSync = false;
+	bool bUsesFileRevisions = false;
 	/** */
 	int32 ChunkIDSelected = 0;
 };
diff --git a/Engine/Plugins/Editor/ContentBrowser/ContentBrowserAssetDataSource/Source/ContentBrowserAssetDataSource/Public/AssetFolderContextMenu.h b/Engine/Plugins/Editor/ContentBrowser/ContentBrowserAssetDataSource/Source/ContentBrowserAssetDataSource/Public/AssetFolderContextMenu.h
--- a/Engine/Plugins/Editor/ContentBrowser/ContentBrowserAssetDataSource/Source/ContentBrowserAssetDataSource/Public/AssetFolderContextMenu.h
+++ b/Engine/Plugins/Editor/ContentBrowser/ContentBrowserAssetDataSource/Source/ContentBrowserAssetDataSource/Public/AssetFolderContextMenu.h
@@ -113,6 +113,8 @@ private:
 	bool bCanExecuteSCCCheckOut = false;
 	bool bCanExecuteSCCOpenForAdd = false;
 	bool bCanExecuteSCCCheckIn = false;
+	bool bUsesFileRevisions = false;
+
 
 	TArray<FString> PackageNamesToProcess;
 	int32 CurrentPackageIndex = 0;
