FResultMessage cleanup

diff --git a/Engine/Plugins/Editor/ContentBrowser/ContentBrowserAssetDataSource/Source/ContentBrowserAssetDataSource/Private/ContentBrowserAssetDataCore.cpp b/Engine/Plugins/Editor/ContentBrowser/ContentBrowserAssetDataSource/Source/ContentBrowserAssetDataSource/Private/ContentBrowserAssetDataCore.cpp
--- a/Engine/Plugins/Editor/ContentBrowser/ContentBrowserAssetDataSource/Source/ContentBrowserAssetDataSource/Private/ContentBrowserAssetDataCore.cpp
+++ b/Engine/Plugins/Editor/ContentBrowser/ContentBrowserAssetDataSource/Source/ContentBrowserAssetDataSource/Private/ContentBrowserAssetDataCore.cpp
@@ -857,11 +857,11 @@ bool RenameAssetFileItem(IAssetTools* InAssetTools, const FContentBrowserAssetFi
 	if (UObject* Asset = InAssetPayload.LoadAsset({ ULevel::LoadAllExternalObjectsTag }))
 	{
 		FResultMessage Result;
-		Result.bSucceeded = true;
+		Result.bSuccess = true;
 		FEditorDelegates::OnPreDestructiveAssetAction.Broadcast({Asset}, EDestructiveAssetActions::AssetRename, Result);
-		if (!Result.WasSuccesful())
+		if (!Result.bSuccess)
 		{
-			UE_LOG(LogContentBrowserAssetDataSource, Warning, TEXT("%s"), *Result.GetErrorMessage());
+			UE_LOG(LogContentBrowserAssetDataSource, Warning, TEXT("%s"), *Result.ErrorMessage);
 			return false;
 		}
 
@@ -1028,11 +1028,11 @@ bool MoveAssetFileItems(TArrayView<const TSharedRef<const FContentBrowserAssetFi
 	}
 
 	FResultMessage Result;
-	Result.bSucceeded = true;
+	Result.bSuccess = true;
 	FEditorDelegates::OnPreDestructiveAssetAction.Broadcast(AssetsToMove, EDestructiveAssetActions::AssetMove, Result);
-	if (!Result.WasSuccesful())
+	if (!Result.bSuccess)
 	{
-		UE_LOG(LogContentBrowserAssetDataSource, Warning, TEXT("%s"), *Result.GetErrorMessage());
+		UE_LOG(LogContentBrowserAssetDataSource, Warning, TEXT("%s"), *Result.ErrorMessage);
 		return false;
 	}
 
diff --git a/Engine/Source/Developer/AssetTools/Private/AssetViewUtils.cpp b/Engine/Source/Developer/AssetTools/Private/AssetViewUtils.cpp
--- a/Engine/Source/Developer/AssetTools/Private/AssetViewUtils.cpp
+++ b/Engine/Source/Developer/AssetTools/Private/AssetViewUtils.cpp
@@ -555,11 +555,11 @@ bool AssetViewUtils::RenameFolder(const FString& DestPath, const FString& Source
 	GetObjectsInAssetData(AssetsInFolder, ObjectsInFolder, bLoadAllExternalObjects);
 
 	FResultMessage Result;
-	Result.bSucceeded = true;
+	Result.bSuccess = true;
 	FEditorDelegates::OnPreDestructiveAssetAction.Broadcast(ObjectsInFolder, EDestructiveAssetActions::AssetRename, Result);
-	if (!Result.WasSuccesful())
+	if (!Result.bSuccess)
 	{
-		UE_LOG(LogAssetViewTools, Warning, TEXT("%s"), *Result.GetErrorMessage());
+		UE_LOG(LogAssetViewTools, Warning, TEXT("%s"), *Result.ErrorMessage);
 		return false;
 	}
 
@@ -670,11 +670,11 @@ bool AssetViewUtils::MoveFolders(const TArray<FString>& InSourcePathNames, const
 	}
 
 	FResultMessage Result;
-	Result.bSucceeded = true;
+	Result.bSuccess = true;
 	FEditorDelegates::OnPreDestructiveAssetAction.Broadcast(AssetsToMove, EDestructiveAssetActions::AssetMove, Result);
-	if (!Result.WasSuccesful())
+	if (!Result.bSuccess)
 	{
-		UE_LOG(LogAssetViewTools, Warning, TEXT("%s"), *Result.GetErrorMessage());
+		UE_LOG(LogAssetViewTools, Warning, TEXT("%s"), *Result.ErrorMessage);
 		return false;
 	}
 	
diff --git a/Engine/Source/Editor/UnrealEd/Classes/Editor/UnrealEdTypes.h b/Engine/Source/Editor/UnrealEd/Classes/Editor/UnrealEdTypes.h
--- a/Engine/Source/Editor/UnrealEd/Classes/Editor/UnrealEdTypes.h
+++ b/Engine/Source/Editor/UnrealEd/Classes/Editor/UnrealEdTypes.h
@@ -139,20 +139,12 @@ enum EDestructiveAssetActions
 };
 
 /** Simple Return struct for supplying success or failure with an optional error message */
-USTRUCT()
 struct FResultMessage
 {
-	GENERATED_USTRUCT_BODY()
-public:
-	bool WasSuccesful() { return bSucceeded; }
-	FString GetErrorMessage() { return ErrorMessage; }
-	void SetErrorMessage(FString& InErrorMessage) { InErrorMessage = ErrorMessage; }
-	bool bSucceeded = false;
-private:
+	bool bSuccess = false;
 	FString ErrorMessage;
 };
 
-
 UCLASS(abstract, config=UnrealEd)
 class UUnrealEdTypes : public UObject
 {
diff --git a/Engine/Source/Editor/UnrealEd/Private/ObjectTools.cpp b/Engine/Source/Editor/UnrealEd/Private/ObjectTools.cpp
--- a/Engine/Source/Editor/UnrealEd/Private/ObjectTools.cpp
+++ b/Engine/Source/Editor/UnrealEd/Private/ObjectTools.cpp
@@ -2688,12 +2688,12 @@ namespace ObjectTools
 		}
 
 		FResultMessage Result;
-		Result.bSucceeded = true;
+		Result.bSuccess = true;
 		FEditorDelegates::OnPreDestructiveAssetAction.Broadcast(ObjectsToDelete, EDestructiveAssetActions::AssetDelete, Result);
 
-		if (!Result.WasSuccesful())
+		if (!Result.bSuccess)
 		{
-			UE_LOG(LogObjectTools, Warning, TEXT("%s"), *Result.GetErrorMessage());
+			UE_LOG(LogObjectTools, Warning, TEXT("%s"), *Result.ErrorMessage);
 			return 0;
 		}
 
@@ -2792,12 +2792,12 @@ namespace ObjectTools
 		}
 
 		FResultMessage Result;
-		Result.bSucceeded = true;
+		Result.bSuccess = true;
 		FEditorDelegates::OnPreDestructiveAssetAction.Broadcast(ObjectsToPrivatize, EDestructiveAssetActions::AssetPrivatize, Result);
 
-		if (!Result.WasSuccesful())
+		if (!Result.bSuccess)
 		{
-			UE_LOG(LogObjectTools, Warning, TEXT("%s"), *Result.GetErrorMessage());
+			UE_LOG(LogObjectTools, Warning, TEXT("%s"), *Result.ErrorMessage);
 			return 0;
 		}
 
