#Ue Fix when right clicking an asset that cannot be loaded

diff --git a/Engine/Plugins/Editor/ContentBrowser/ContentBrowserAssetDataSource/Source/ContentBrowserAssetDataSource/Private/AssetFileContextMenu.cpp b/Engine/Plugins/Editor/ContentBrowser/ContentBrowserAssetDataSource/Source/ContentBrowserAssetDataSource/Private/AssetFileContextMenu.cpp
--- a/Engine/Plugins/Editor/ContentBrowser/ContentBrowserAssetDataSource/Source/ContentBrowserAssetDataSource/Private/AssetFileContextMenu.cpp
+++ b/Engine/Plugins/Editor/ContentBrowser/ContentBrowserAssetDataSource/Source/ContentBrowserAssetDataSource/Private/AssetFileContextMenu.cpp
@@ -1215,7 +1215,10 @@ void FAssetFileContextMenu::ExecuteReimport(int32 SourceFileIndex /*= INDEX_NONE
 	for (const FAssetData &SelectedAsset : SelectedAssets)
 	{
 		UObject *Asset = SelectedAsset.GetAsset();
-		CopyOfSelectedAssets.Add(Asset);
+		if (Asset)
+		{
+			CopyOfSelectedAssets.Add(Asset);
+		}
 	}
 	FReimportManager::Instance()->ValidateAllSourceFileAndReimport(CopyOfSelectedAssets, true, SourceFileIndex, false);
 }
@@ -1229,7 +1232,10 @@ void FAssetFileContextMenu::ExecuteReimportWithNewFile(int32 SourceFileIndex /*=
 	for (const FAssetData &SelectedAsset : SelectedAssets)
 	{
 		UObject *Asset = SelectedAsset.GetAsset();
-		CopyOfSelectedAssets.Add(Asset);
+		if (Asset)
+		{
+			CopyOfSelectedAssets.Add(Asset);
+		}
 	}
 
 	TArray<FString> AssetSourcePaths;
@@ -1283,14 +1289,17 @@ void FAssetFileContextMenu::GetSelectedAssetsByClass(TMap<UClass*, TArray<UObjec
 	for (const auto& SelectedAsset : SelectedAssets)
 	{
 		auto Asset = SelectedAsset.GetAsset();
-		auto AssetClass = Asset->GetClass();
-
-		if ( !OutSelectedAssetsByClass.Contains(AssetClass) )
+		if (Asset)
 		{
-			OutSelectedAssetsByClass.Add(AssetClass);
-		}
+			auto AssetClass = Asset->GetClass();
+
+			if ( !OutSelectedAssetsByClass.Contains(AssetClass) )
+			{
+				OutSelectedAssetsByClass.Add(AssetClass);
+			}
 		
-		OutSelectedAssetsByClass[AssetClass].Add(Asset);
+			OutSelectedAssetsByClass[AssetClass].Add(Asset);
+		}
 	}
 }
 
