Fixed Content Browser class permission list filtering always being recursive

diff --git a/Engine/Plugins/Editor/ContentBrowser/ContentBrowserAssetDataSource/Source/ContentBrowserAssetDataSource/Private/ContentBrowserAssetDataSource.cpp b/Engine/Plugins/Editor/ContentBrowser/ContentBrowserAssetDataSource/Source/ContentBrowserAssetDataSource/Private/ContentBrowserAssetDataSource.cpp
--- a/Engine/Plugins/Editor/ContentBrowser/ContentBrowserAssetDataSource/Source/ContentBrowserAssetDataSource/Private/ContentBrowserAssetDataSource.cpp
+++ b/Engine/Plugins/Editor/ContentBrowser/ContentBrowserAssetDataSource/Source/ContentBrowserAssetDataSource/Private/ContentBrowserAssetDataSource.cpp
@@ -557,7 +557,6 @@ PRAGMA_ENABLE_DEPRECATION_WARNINGS
 				{
 					AllowListClassFilter.ClassPaths.Add(FTopLevelAssetPath(AllowListPair.Key));
 				}
-				AllowListClassFilter.bRecursiveClasses = true;
 				Params.AssetRegistry->CompileFilter(AllowListClassFilter, CompiledClassFilterAllowList);
 			}
 
@@ -634,7 +633,6 @@ PRAGMA_ENABLE_DEPRECATION_WARNINGS
 				{
 					ClassFilter.ClassPaths.Add(FTopLevelAssetPath(FilterPair.Key));
 				}
-				ClassFilter.bRecursiveClasses = true;
 				Params.AssetRegistry->CompileFilter(ClassFilter, CompiledClassFilter);
 			}
 
