Remove registration of AssetRegistry notification in content browser when running as commandlet.

I noticed a fair amount of notifications and processing in response to those notifications from the AssetRegistry when running a cook command without the editor UI.

This CL removes registration of callbacks if we are running a commandlet such as cooking.

diff --git a/Engine/Plugins/Editor/ContentBrowser/ContentBrowserAliasDataSoure/Source/ContentBrowserAliasDataSource/Private/ContentBrowserAliasDataSource.cpp b/Engine/Plugins/Editor/ContentBrowser/ContentBrowserAliasDataSoure/Source/ContentBrowserAliasDataSource/Private/ContentBrowserAliasDataSource.cpp
--- a/Engine/Plugins/Editor/ContentBrowser/ContentBrowserAliasDataSoure/Source/ContentBrowserAliasDataSource/Private/ContentBrowserAliasDataSource.cpp
+++ b/Engine/Plugins/Editor/ContentBrowser/ContentBrowserAliasDataSoure/Source/ContentBrowserAliasDataSource/Private/ContentBrowserAliasDataSource.cpp
@@ -23,15 +23,17 @@ FName UContentBrowserAliasDataSource::AliasTagName = "ContentBrowserAliases";
 void UContentBrowserAliasDataSource::Initialize(const bool InAutoRegister)
 {
 	Super::Initialize(InAutoRegister);
+	if (GIsEditor && !IsRunningCommandlet())
+	{
+		AssetRegistry = &FModuleManager::LoadModuleChecked<FAssetRegistryModule>(AssetRegistryConstants::ModuleName).Get();
+		AssetTools = &FModuleManager::GetModuleChecked<FAssetToolsModule>("AssetTools").Get();
 
-	AssetRegistry = &FModuleManager::LoadModuleChecked<FAssetRegistryModule>(AssetRegistryConstants::ModuleName).Get();
-	AssetTools = &FModuleManager::GetModuleChecked<FAssetToolsModule>("AssetTools").Get();
-
-	AssetRegistry->OnAssetAdded().AddUObject(this, &UContentBrowserAliasDataSource::OnAssetAdded);
-	AssetRegistry->OnAssetRemoved().AddUObject(this, &UContentBrowserAliasDataSource::OnAssetRemoved);
-	AssetRegistry->OnAssetUpdated().AddUObject(this, &UContentBrowserAliasDataSource::OnAssetUpdated);
-	FCoreUObjectDelegates::OnAssetLoaded.AddUObject(this, &UContentBrowserAliasDataSource::OnAssetLoaded);
-	FCoreUObjectDelegates::OnObjectPropertyChanged.AddUObject(this, &UContentBrowserAliasDataSource::OnObjectPropertyChanged);
+		AssetRegistry->OnAssetAdded().AddUObject(this, &UContentBrowserAliasDataSource::OnAssetAdded);
+		AssetRegistry->OnAssetRemoved().AddUObject(this, &UContentBrowserAliasDataSource::OnAssetRemoved);
+		AssetRegistry->OnAssetUpdated().AddUObject(this, &UContentBrowserAliasDataSource::OnAssetUpdated);
+		FCoreUObjectDelegates::OnAssetLoaded.AddUObject(this, &UContentBrowserAliasDataSource::OnAssetLoaded);
+		FCoreUObjectDelegates::OnObjectPropertyChanged.AddUObject(this, &UContentBrowserAliasDataSource::OnObjectPropertyChanged);
+	}
 }
 
 void UContentBrowserAliasDataSource::Shutdown()
diff --git a/Engine/Plugins/Editor/ContentBrowser/ContentBrowserAssetDataSource/Source/ContentBrowserAssetDataSource/Private/ContentBrowserAssetDataSource.cpp b/Engine/Plugins/Editor/ContentBrowser/ContentBrowserAssetDataSource/Source/ContentBrowserAssetDataSource/Private/ContentBrowserAssetDataSource.cpp
--- a/Engine/Plugins/Editor/ContentBrowser/ContentBrowserAssetDataSource/Source/ContentBrowserAssetDataSource/Private/ContentBrowserAssetDataSource.cpp
+++ b/Engine/Plugins/Editor/ContentBrowser/ContentBrowserAssetDataSource/Source/ContentBrowserAssetDataSource/Private/ContentBrowserAssetDataSource.cpp
@@ -42,123 +42,126 @@ void UContentBrowserAssetDataSource::Initialize(const bool InAutoRegister)
 {
 	Super::Initialize(InAutoRegister);
 
-	AssetRegistry = &FModuleManager::LoadModuleChecked<FAssetRegistryModule>(AssetRegistryConstants::ModuleName).Get();
-	AssetRegistry->OnFileLoadProgressUpdated().AddUObject(this, &UContentBrowserAssetDataSource::OnAssetRegistryFileLoadProgress);
-
+	if (GIsEditor && !IsRunningCommandlet())
 	{
-		static const FName NAME_AssetTools = "AssetTools";
-		AssetTools = &FModuleManager::GetModuleChecked<FAssetToolsModule>(NAME_AssetTools).Get();
-	}
+		AssetRegistry = &FModuleManager::LoadModuleChecked<FAssetRegistryModule>(AssetRegistryConstants::ModuleName).Get();
+		AssetRegistry->OnFileLoadProgressUpdated().AddUObject(this, &UContentBrowserAssetDataSource::OnAssetRegistryFileLoadProgress);
+
+		{
+			static const FName NAME_AssetTools = "AssetTools";
+			AssetTools = &FModuleManager::GetModuleChecked<FAssetToolsModule>(NAME_AssetTools).Get();
+		}
 
-	CollectionManager = &FCollectionManagerModule::GetModule().Get();
+		CollectionManager = &FCollectionManagerModule::GetModule().Get();
 
-	// Listen for asset registry updates
-	AssetRegistry->OnAssetAdded().AddUObject(this, &UContentBrowserAssetDataSource::OnAssetAdded);
-	AssetRegistry->OnAssetRemoved().AddUObject(this, &UContentBrowserAssetDataSource::OnAssetRemoved);
-	AssetRegistry->OnAssetRenamed().AddUObject(this, &UContentBrowserAssetDataSource::OnAssetRenamed);
-	AssetRegistry->OnAssetUpdated().AddUObject(this, &UContentBrowserAssetDataSource::OnAssetUpdated);
-	AssetRegistry->OnPathAdded().AddUObject(this, &UContentBrowserAssetDataSource::OnPathAdded);
-	AssetRegistry->OnPathRemoved().AddUObject(this, &UContentBrowserAssetDataSource::OnPathRemoved);
-	AssetRegistry->OnFilesLoaded().AddUObject(this, &UContentBrowserAssetDataSource::OnScanCompleted);
+		// Listen for asset registry updates
+		AssetRegistry->OnAssetAdded().AddUObject(this, &UContentBrowserAssetDataSource::OnAssetAdded);
+		AssetRegistry->OnAssetRemoved().AddUObject(this, &UContentBrowserAssetDataSource::OnAssetRemoved);
+		AssetRegistry->OnAssetRenamed().AddUObject(this, &UContentBrowserAssetDataSource::OnAssetRenamed);
+		AssetRegistry->OnAssetUpdated().AddUObject(this, &UContentBrowserAssetDataSource::OnAssetUpdated);
+		AssetRegistry->OnPathAdded().AddUObject(this, &UContentBrowserAssetDataSource::OnPathAdded);
+		AssetRegistry->OnPathRemoved().AddUObject(this, &UContentBrowserAssetDataSource::OnPathRemoved);
+		AssetRegistry->OnFilesLoaded().AddUObject(this, &UContentBrowserAssetDataSource::OnScanCompleted);
 
-	// Listen for when assets are loaded or changed
-	FCoreUObjectDelegates::OnObjectPropertyChanged.AddUObject(this, &UContentBrowserAssetDataSource::OnObjectPropertyChanged);
+		// Listen for when assets are loaded or changed
+		FCoreUObjectDelegates::OnObjectPropertyChanged.AddUObject(this, &UContentBrowserAssetDataSource::OnObjectPropertyChanged);
 
-	// Listen for new mount roots
-	FPackageName::OnContentPathMounted().AddUObject(this, &UContentBrowserAssetDataSource::OnContentPathMounted);
-	FPackageName::OnContentPathDismounted().AddUObject(this, &UContentBrowserAssetDataSource::OnContentPathDismounted);
+		// Listen for new mount roots
+		FPackageName::OnContentPathMounted().AddUObject(this, &UContentBrowserAssetDataSource::OnContentPathMounted);
+		FPackageName::OnContentPathDismounted().AddUObject(this, &UContentBrowserAssetDataSource::OnContentPathDismounted);
 
-	// Listen for paths being forced visible
-	AssetViewUtils::OnAlwaysShowPath().AddUObject(this, &UContentBrowserAssetDataSource::OnAlwaysShowPath);
+		// Listen for paths being forced visible
+		AssetViewUtils::OnAlwaysShowPath().AddUObject(this, &UContentBrowserAssetDataSource::OnAlwaysShowPath);
 
-	// Register our ability to create assets via the legacy Content Browser API
-	ContentBrowserDataLegacyBridge::OnCreateNewAsset().BindUObject(this, &UContentBrowserAssetDataSource::OnBeginCreateAsset);
+		// Register our ability to create assets via the legacy Content Browser API
+		ContentBrowserDataLegacyBridge::OnCreateNewAsset().BindUObject(this, &UContentBrowserAssetDataSource::OnBeginCreateAsset);
 
-	// Create the asset menu instances
-	AssetFolderContextMenu = MakeShared<FAssetFolderContextMenu>();
-	AssetFileContextMenu = MakeShared<FAssetFileContextMenu>();
+		// Create the asset menu instances
+		AssetFolderContextMenu = MakeShared<FAssetFolderContextMenu>();
+		AssetFileContextMenu = MakeShared<FAssetFileContextMenu>();
 
-	// Bind the asset specific menu extensions
-	{
-		if (UToolMenu* Menu = UToolMenus::Get()->ExtendMenu("ContentBrowser.AddNewContextMenu"))
+		// Bind the asset specific menu extensions
 		{
-			Menu->AddDynamicSection(*FString::Printf(TEXT("DynamicSection_DataSource_%s"), *GetName()), FNewToolMenuDelegate::CreateLambda([WeakThis = TWeakObjectPtr<UContentBrowserAssetDataSource>(this)](UToolMenu* InMenu)
+			if (UToolMenu* Menu = UToolMenus::Get()->ExtendMenu("ContentBrowser.AddNewContextMenu"))
 			{
-				if (UContentBrowserAssetDataSource* This = WeakThis.Get())
-				{
-					This->PopulateAddNewContextMenu(InMenu);
-				}
-			}));
-		}
+				Menu->AddDynamicSection(*FString::Printf(TEXT("DynamicSection_DataSource_%s"), *GetName()), FNewToolMenuDelegate::CreateLambda([WeakThis = TWeakObjectPtr<UContentBrowserAssetDataSource>(this)](UToolMenu* InMenu)
+					{
+						if (UContentBrowserAssetDataSource* This = WeakThis.Get())
+						{
+							This->PopulateAddNewContextMenu(InMenu);
+						}
+					}));
+			}
 
 
-		if (UToolMenu* Menu = UToolMenus::Get()->ExtendMenu("ContentBrowser.ToolBar"))
-		{
-			Menu->AddDynamicSection(*FString::Printf(TEXT("DynamicSection_DataSource_%s"), *GetName()), FNewToolMenuDelegate::CreateLambda([WeakThis = TWeakObjectPtr<UContentBrowserAssetDataSource>(this)](UToolMenu* InMenu)
+			if (UToolMenu* Menu = UToolMenus::Get()->ExtendMenu("ContentBrowser.ToolBar"))
 			{
-				if (UContentBrowserAssetDataSource* This = WeakThis.Get())
-				{
-					This->PopulateContentBrowserToolBar(InMenu);
-				}
-			}));
-		}
+				Menu->AddDynamicSection(*FString::Printf(TEXT("DynamicSection_DataSource_%s"), *GetName()), FNewToolMenuDelegate::CreateLambda([WeakThis = TWeakObjectPtr<UContentBrowserAssetDataSource>(this)](UToolMenu* InMenu)
+					{
+						if (UContentBrowserAssetDataSource* This = WeakThis.Get())
+						{
+							This->PopulateContentBrowserToolBar(InMenu);
+						}
+					}));
+			}
 
 
-		if (UToolMenu* Menu = UToolMenus::Get()->ExtendMenu("ContentBrowser.FolderContextMenu"))
-		{
-			Menu->AddDynamicSection(*FString::Printf(TEXT("DynamicSection_DataSource_%s"), *GetName()), FNewToolMenuDelegate::CreateLambda([WeakThis = TWeakObjectPtr<UContentBrowserAssetDataSource>(this)](UToolMenu* InMenu)
+			if (UToolMenu* Menu = UToolMenus::Get()->ExtendMenu("ContentBrowser.FolderContextMenu"))
 			{
-				if (UContentBrowserAssetDataSource* This = WeakThis.Get())
-				{
-					This->PopulateAssetFolderContextMenu(InMenu);
-				}
-			}));
-		}
+				Menu->AddDynamicSection(*FString::Printf(TEXT("DynamicSection_DataSource_%s"), *GetName()), FNewToolMenuDelegate::CreateLambda([WeakThis = TWeakObjectPtr<UContentBrowserAssetDataSource>(this)](UToolMenu* InMenu)
+					{
+						if (UContentBrowserAssetDataSource* This = WeakThis.Get())
+						{
+							This->PopulateAssetFolderContextMenu(InMenu);
+						}
+					}));
+			}
 
-		if (UToolMenu* Menu = UToolMenus::Get()->ExtendMenu("ContentBrowser.AssetContextMenu"))
-		{
-			Menu->AddDynamicSection(*FString::Printf(TEXT("DynamicSection_DataSource_%s"), *GetName()), FNewToolMenuDelegate::CreateLambda([WeakThis = TWeakObjectPtr<UContentBrowserAssetDataSource>(this)](UToolMenu* InMenu)
+			if (UToolMenu* Menu = UToolMenus::Get()->ExtendMenu("ContentBrowser.AssetContextMenu"))
 			{
-				if (UContentBrowserAssetDataSource* This = WeakThis.Get())
-				{
-					This->PopulateAssetFileContextMenu(InMenu);
-				}
-			}));
-		}
+				Menu->AddDynamicSection(*FString::Printf(TEXT("DynamicSection_DataSource_%s"), *GetName()), FNewToolMenuDelegate::CreateLambda([WeakThis = TWeakObjectPtr<UContentBrowserAssetDataSource>(this)](UToolMenu* InMenu)
+					{
+						if (UContentBrowserAssetDataSource* This = WeakThis.Get())
+						{
+							This->PopulateAssetFileContextMenu(InMenu);
+						}
+					}));
+			}
 
-		if (UToolMenu* Menu = UToolMenus::Get()->ExtendMenu("ContentBrowser.DragDropContextMenu"))
-		{
-			Menu->AddDynamicSection(*FString::Printf(TEXT("DynamicSection_DataSource_%s"), *GetName()), FNewToolMenuDelegate::CreateLambda([WeakThis = TWeakObjectPtr<UContentBrowserAssetDataSource>(this)](UToolMenu* InMenu)
+			if (UToolMenu* Menu = UToolMenus::Get()->ExtendMenu("ContentBrowser.DragDropContextMenu"))
 			{
-				if (UContentBrowserAssetDataSource* This = WeakThis.Get())
-				{
-					This->PopulateDragDropContextMenu(InMenu);
-				}
-			}));
+				Menu->AddDynamicSection(*FString::Printf(TEXT("DynamicSection_DataSource_%s"), *GetName()), FNewToolMenuDelegate::CreateLambda([WeakThis = TWeakObjectPtr<UContentBrowserAssetDataSource>(this)](UToolMenu* InMenu)
+					{
+						if (UContentBrowserAssetDataSource* This = WeakThis.Get())
+						{
+							This->PopulateDragDropContextMenu(InMenu);
+						}
+					}));
+			}
 		}
-	}
 
-	DiscoveryStatusText = LOCTEXT("InitializingAssetDiscovery", "Initializing Asset Discovery...");
+		DiscoveryStatusText = LOCTEXT("InitializingAssetDiscovery", "Initializing Asset Discovery...");
 
-	// Populate the initial set of hidden empty folders
-	// This will be updated as the scan finds more content
-	AssetRegistry->EnumerateAllCachedPaths([this](FName InPath)
-	{
-		if (!AssetRegistry->HasAssets(InPath, /*bRecursive*/true))
-		{
-			EmptyAssetFolders.Add(InPath.ToString());
-		}
+		// Populate the initial set of hidden empty folders
+		// This will be updated as the scan finds more content
+		AssetRegistry->EnumerateAllCachedPaths([this](FName InPath)
+			{
+				if (!AssetRegistry->HasAssets(InPath, /*bRecursive*/true))
+				{
+					EmptyAssetFolders.Add(InPath.ToString());
+				}
 		return true;
-	});
+			});
 
-	FPackageName::QueryRootContentPaths(RootContentPaths);
+		FPackageName::QueryRootContentPaths(RootContentPaths);
 
-	BuildRootPathVirtualTree();
+		BuildRootPathVirtualTree();
 
-	// Mount roots are always visible
-	for (const FString& RootContentPath : RootContentPaths)
-	{
-		OnAlwaysShowPath(RootContentPath);
+		// Mount roots are always visible
+		for (const FString& RootContentPath : RootContentPaths)
+		{
+			OnAlwaysShowPath(RootContentPath);
+		}
 	}
 }
 
@@ -167,22 +170,23 @@ void UContentBrowserAssetDataSource::Shutdown()
 	CollectionManager = nullptr;
 
 	AssetTools = nullptr;
+	AssetRegistry = nullptr;
 
-	if (!FModuleManager::Get().IsModuleLoaded(AssetRegistryConstants::ModuleName))
+	if (FAssetRegistryModule* AssetRegistryModule = FModuleManager::GetModulePtr<FAssetRegistryModule>(AssetRegistryConstants::ModuleName))
 	{
-		AssetRegistry = nullptr;
-	}
-
-	if (AssetRegistry)
-	{
-		AssetRegistry->OnFileLoadProgressUpdated().RemoveAll(this);
+		IAssetRegistry* AssetRegistryMaybe = AssetRegistryModule->TryGet();
+		if (AssetRegistryMaybe)
+		{
+			AssetRegistryMaybe->OnFileLoadProgressUpdated().RemoveAll(this);
 
-		AssetRegistry->OnAssetAdded().RemoveAll(this);
-		AssetRegistry->OnAssetRemoved().RemoveAll(this);
-		AssetRegistry->OnAssetRenamed().RemoveAll(this);
-		AssetRegistry->OnAssetUpdated().RemoveAll(this);
-		AssetRegistry->OnPathAdded().RemoveAll(this);
-		AssetRegistry->OnPathRemoved().RemoveAll(this);
+			AssetRegistryMaybe->OnAssetAdded().RemoveAll(this);
+			AssetRegistryMaybe->OnAssetRemoved().RemoveAll(this);
+			AssetRegistryMaybe->OnAssetRenamed().RemoveAll(this);
+			AssetRegistryMaybe->OnAssetUpdated().RemoveAll(this);
+			AssetRegistryMaybe->OnPathAdded().RemoveAll(this);
+			AssetRegistryMaybe->OnPathRemoved().RemoveAll(this);
+			AssetRegistryMaybe->OnFilesLoaded().RemoveAll(this);
+		}
 	}
 
 	FCoreUObjectDelegates::OnObjectPropertyChanged.RemoveAll(this);
