Cached which folders have been visited.

Before: 3s 272ms
After: 1s 577ms

diff --git a/Engine/Plugins/Editor/ContentBrowser/ContentBrowserAssetDataSource/Source/ContentBrowserAssetDataSource/Private/ContentBrowserAssetDataSource.cpp b/Engine/Plugins/Editor/ContentBrowser/ContentBrowserAssetDataSource/Source/ContentBrowserAssetDataSource/Private/ContentBrowserAssetDataSource.cpp
--- a/Engine/Plugins/Editor/ContentBrowser/ContentBrowserAssetDataSource/Source/ContentBrowserAssetDataSource/Private/ContentBrowserAssetDataSource.cpp
+++ b/Engine/Plugins/Editor/ContentBrowser/ContentBrowserAssetDataSource/Source/ContentBrowserAssetDataSource/Private/ContentBrowserAssetDataSource.cpp
@@ -148,10 +148,11 @@ void UContentBrowserAssetDataSource::Initialize(const bool InAutoRegister)
 			{
 				if (!AssetRegistry->HasAssets(InPath, /*bRecursive*/true))
 				{
-					EmptyAssetFolders.Add(InPath.ToString());
+					EmptyAssetFolders.Add(InPath);
 				}
-		return true;
+				return true;
 			});
+		VisitedEmptyAssetFolders.Empty();
 
 		FPackageName::QueryRootContentPaths(RootContentPaths);
 
@@ -1049,13 +1050,8 @@ bool UContentBrowserAssetDataSource::IsFolderVisibleIfHidingEmpty(const FName In
 		return false;
 	}
 
-	FNameBuilder InternalPathStr(ConvertedPath);
-
-	const FStringView InternalPathStrView = InternalPathStr;
-	const uint32 InternalPathHash = GetTypeHash(InternalPathStrView);
-
-	return AlwaysVisibleAssetFolders.ContainsByHash(InternalPathHash, InternalPathStrView) 
-		|| !EmptyAssetFolders.ContainsByHash(InternalPathHash, InternalPathStrView);
+	return AlwaysVisibleAssetFolders.Contains(ConvertedPath)
+		|| !EmptyAssetFolders.Contains(ConvertedPath);
 }
 
 bool UContentBrowserAssetDataSource::CanCreateFolder(const FName InPath, FText* OutErrorMsg)
@@ -1803,18 +1799,21 @@ void UContentBrowserAssetDataSource::OnObjectPropertyChanged(UObject* InObject,
 void UContentBrowserAssetDataSource::OnPathAdded(const FString& InPath)
 {
 	// New paths are considered empty until assets are added inside them
-	EmptyAssetFolders.Add(InPath);
+	FName PathName(InPath);
+	EmptyAssetFolders.Add(PathName);
+	VisitedEmptyAssetFolders.Empty();
 
-	QueueItemDataUpdate(FContentBrowserItemDataUpdate::MakeItemAddedUpdate(CreateAssetFolderItem(*InPath)));
+	QueueItemDataUpdate(FContentBrowserItemDataUpdate::MakeItemAddedUpdate(CreateAssetFolderItem(PathName)));
 }
 
 void UContentBrowserAssetDataSource::OnPathRemoved(const FString& InPath)
 {
 	// Deleted paths are no longer relevant for tracking
-	AlwaysVisibleAssetFolders.Remove(InPath);
-	EmptyAssetFolders.Remove(InPath);
+	FName PathName(InPath);
+	AlwaysVisibleAssetFolders.Remove(PathName);
+	EmptyAssetFolders.Remove(PathName);
 
-	QueueItemDataUpdate(FContentBrowserItemDataUpdate::MakeItemRemovedUpdate(CreateAssetFolderItem(*InPath)));
+	QueueItemDataUpdate(FContentBrowserItemDataUpdate::MakeItemRemovedUpdate(CreateAssetFolderItem(PathName)));
 }
 
 void UContentBrowserAssetDataSource::OnPathPopulated(const FName InPath)
@@ -1834,6 +1833,15 @@ void UContentBrowserAssetDataSource::OnPathPopulated(const FStringView InPath)
 			Path = Path.Left(Path.Len() - 1);
 		}
 
+		FName PathName(Path);
+
+		// If we've already visited this path then we can assume we visited the parents as well
+		// and can skip visiting this path and its parents
+		if (VisitedEmptyAssetFolders.Contains(PathName))
+		{
+			return;
+		}
+
 		// Recurse first as we want parents to be updated before their children
 		{
 			int32 LastSlashIndex = INDEX_NONE;
@@ -1844,12 +1852,14 @@ void UContentBrowserAssetDataSource::OnPathPopulated(const FStringView InPath)
 		}
 
 		// Unhide this folder and emit a notification if required
-		const uint32 PathHash = GetTypeHash(Path);
-		if (EmptyAssetFolders.RemoveByHash(PathHash, Path) > 0)
+		if (EmptyAssetFolders.Remove(PathName) > 0)
 		{
 			// Queue an update event for this path as it may have become visible in the view
-			QueueItemDataUpdate(FContentBrowserItemDataUpdate::MakeItemModifiedUpdate(CreateAssetFolderItem(FName(Path))));
+			QueueItemDataUpdate(FContentBrowserItemDataUpdate::MakeItemModifiedUpdate(CreateAssetFolderItem(PathName)));
 		}
+
+		// Mark that this path has been visited
+		VisitedEmptyAssetFolders.Add(PathName);
 	}
 }
 
@@ -1875,12 +1885,13 @@ void UContentBrowserAssetDataSource::OnAlwaysShowPath(const FString& InPath)
 		}
 
 		// Force show this folder and emit a notification if required
-		if (!AlwaysVisibleAssetFolders.Contains(Path))
+		FName PathName(Path);
+		if (!AlwaysVisibleAssetFolders.Contains(PathName))
 		{
-			AlwaysVisibleAssetFolders.Add(Path);
+			AlwaysVisibleAssetFolders.Add(PathName);
 
 			// Queue an update event for this path as it may have become visible in the view
-			QueueItemDataUpdate(FContentBrowserItemDataUpdate::MakeItemModifiedUpdate(CreateAssetFolderItem(FName(*Path))));
+			QueueItemDataUpdate(FContentBrowserItemDataUpdate::MakeItemModifiedUpdate(CreateAssetFolderItem(PathName)));
 		}
 	}
 }
diff --git a/Engine/Plugins/Editor/ContentBrowser/ContentBrowserAssetDataSource/Source/ContentBrowserAssetDataSource/Public/ContentBrowserAssetDataSource.h b/Engine/Plugins/Editor/ContentBrowser/ContentBrowserAssetDataSource/Source/ContentBrowserAssetDataSource/Public/ContentBrowserAssetDataSource.h
--- a/Engine/Plugins/Editor/ContentBrowser/ContentBrowserAssetDataSource/Source/ContentBrowserAssetDataSource/Public/ContentBrowserAssetDataSource.h
+++ b/Engine/Plugins/Editor/ContentBrowser/ContentBrowserAssetDataSource/Source/ContentBrowserAssetDataSource/Public/ContentBrowserAssetDataSource.h
@@ -364,10 +364,15 @@ private:
 	 * The set of folders that should always be visible, even if they contain no assets in the Content Browser view.
 	 * This will include root content folders, and any folders that have been created directly (or indirectly) by a user action.
 	 */
-	TSet<FString> AlwaysVisibleAssetFolders;
+	TSet<FName> AlwaysVisibleAssetFolders;
 
 	/**
 	 * A cache of folders that contain no assets in the Content Browser view.
 	 */
-	TSet<FString> EmptyAssetFolders;
+	TSet<FName> EmptyAssetFolders;
+	
+	/**
+	 * A cache of folders that visited since the last time any empty asset folders were added.
+	 */
+	TSet<FName> VisitedEmptyAssetFolders;
 };
