#UE Do not load assets when opening the Asset Actions submenu. I had to remove CanExecuteShowAssetMetaData since there was not a good way to check this on unloaded assets. The user will get a notification for any assets that did not have metadata when they execute the command.

diff --git a/Engine/Plugins/Editor/ContentBrowser/ContentBrowserAssetDataSource/Source/ContentBrowserAssetDataSource/Private/AssetFileContextMenu.cpp b/Engine/Plugins/Editor/ContentBrowser/ContentBrowserAssetDataSource/Source/ContentBrowserAssetDataSource/Private/AssetFileContextMenu.cpp
--- a/Engine/Plugins/Editor/ContentBrowser/ContentBrowserAssetDataSource/Source/ContentBrowserAssetDataSource/Private/AssetFileContextMenu.cpp
+++ b/Engine/Plugins/Editor/ContentBrowser/ContentBrowserAssetDataSource/Source/ContentBrowserAssetDataSource/Private/AssetFileContextMenu.cpp
@@ -509,10 +509,7 @@ void FAssetFileContextMenu::MakeAssetActionsSubMenu(UToolMenu* Menu)
 			LOCTEXT("ShowAssetMetaData", "Show Metadata"),
 			LOCTEXT("ShowAssetMetaDataTooltip", "Show the asset metadata dialog."),
 			FSlateIcon(),
-			FUIAction(
-				FExecuteAction::CreateSP(this, &FAssetFileContextMenu::ExecuteShowAssetMetaData),
-				FCanExecuteAction::CreateSP(this, &FAssetFileContextMenu::CanExecuteShowAssetMetaData)
-			)
+			FUIAction( FExecuteAction::CreateSP(this, &FAssetFileContextMenu::ExecuteShowAssetMetaData) )
 		);
 
 		// Chunk actions
@@ -1349,6 +1346,29 @@ void FAssetFileContextMenu::GetSelectedAssets(TArray<UObject*>& Assets, bool Ski
 	AssetViewUtils::LoadAssetsIfNeeded(SelectedAssetPaths, Assets);
 }
 
+void FAssetFileContextMenu::GetSelectedAssetData(TArray<FAssetData>& AssetDataList, bool SkipRedirectors) const
+{
+	if (SkipRedirectors)
+	{
+		TArray<FString> SelectedAssetPaths;
+		AssetDataList.Reserve(SelectedAssets.Num());
+		for (const FAssetData& SelectedAsset : SelectedAssets)
+		{
+			if (SkipRedirectors && (SelectedAsset.AssetClassPath == UObjectRedirector::StaticClass()->GetClassPathName()))
+			{
+				// Don't operate on Redirectors
+				continue;
+			}
+
+			AssetDataList.Add(SelectedAsset);
+		}
+	}
+	else
+	{
+		AssetDataList = SelectedAssets;
+	}
+}
+
 /** Generates a reference graph of the world and can then find actors referencing specified objects */
 struct WorldReferenceGenerator : public FFindReferencedAssets
 {
@@ -1532,6 +1552,12 @@ void FAssetFileContextMenu::ExecuteShowAssetMetaData()
 
 				FSlateApplication::Get().AddWindow(Window.ToSharedRef());
 			}
+			else
+			{
+				FNotificationInfo Info(FText::Format(LOCTEXT("NoMetaDataFound", "No metadata found for asset {0}."), FText::FromString(Asset->GetName())));
+				Info.ExpireDuration = 3.0f;
+				FSlateNotificationManager::Get().AddNotification(Info);
+			}
 		}
 	}
 }
@@ -1891,16 +1917,18 @@ bool FAssetFileContextMenu::CanExecutePropertyMatrix(FText& OutErrorMessage) con
 	bool bResult = bAtLeastOneNonRedirectorSelected;
 	if (bAtLeastOneNonRedirectorSelected)
 	{
-		TArray<UObject*> ObjectsForPropertiesMenu;
 		const bool SkipRedirectors = true;
-		GetSelectedAssets(ObjectsForPropertiesMenu, SkipRedirectors);
+		TArray<FAssetData> AssetDataList;
+		GetSelectedAssetData(AssetDataList, SkipRedirectors);
 
 		// Ensure all Blueprints are valid.
-		for (UObject* Object : ObjectsForPropertiesMenu)
+		static FName GeneratedClassName = TEXT("GeneratedClass");
+		for (const FAssetData& AssetData : AssetDataList)
 		{
-			if (UBlueprint* BlueprintObj = Cast<UBlueprint>(Object))
+			FString GeneratedClassValue;
+			if (AssetData.GetTagValue(GeneratedClassName, GeneratedClassValue))
 			{
-				if (BlueprintObj->GeneratedClass == nullptr)
+				if (GeneratedClassValue.IsEmpty() || GeneratedClassValue == TEXT("None"))
 				{
 					OutErrorMessage = LOCTEXT("InvalidBlueprint", "A selected Blueprint is invalid.");
 					bResult = false;
@@ -1928,24 +1956,6 @@ FText FAssetFileContextMenu::GetExecutePropertyMatrixTooltip() const
 	return ResultTooltip;
 }
 
-bool FAssetFileContextMenu::CanExecuteShowAssetMetaData() const
-{
-	TArray<UObject*> ObjectsForPropertiesMenu;
-	const bool SkipRedirectors = true;
-	GetSelectedAssets(ObjectsForPropertiesMenu, SkipRedirectors);
-
-	bool bResult = false;
-	for (const UObject* Asset : ObjectsForPropertiesMenu)
-	{
-		if (Asset && UMetaData::GetMapForObject(Asset))
-		{
-			bResult = true;
-			break;
-		}
-	}
-	return bResult;
-}
-
 bool FAssetFileContextMenu::CanExecuteCaptureThumbnail() const
 {
 	return GCurrentLevelEditingViewportClient != NULL;
diff --git a/Engine/Plugins/Editor/ContentBrowser/ContentBrowserAssetDataSource/Source/ContentBrowserAssetDataSource/Public/AssetFileContextMenu.h b/Engine/Plugins/Editor/ContentBrowser/ContentBrowserAssetDataSource/Source/ContentBrowserAssetDataSource/Public/AssetFileContextMenu.h
--- a/Engine/Plugins/Editor/ContentBrowser/ContentBrowserAssetDataSource/Source/ContentBrowserAssetDataSource/Public/AssetFileContextMenu.h
+++ b/Engine/Plugins/Editor/ContentBrowser/ContentBrowserAssetDataSource/Source/ContentBrowserAssetDataSource/Public/AssetFileContextMenu.h
@@ -203,9 +203,6 @@ private:
 	bool CanExecutePropertyMatrix(FText& OutErrorMessage) const;
 	bool CanExecutePropertyMatrix() const;
 
-	/** Handler to check to see if a "Show MetaData" command is allowed */
-	bool CanExecuteShowAssetMetaData() const;
-
 	/** Handler to check to see if "Capture Thumbnail" can be executed */
 	bool CanExecuteCaptureThumbnail() const;
 
@@ -239,6 +236,9 @@ private:
 	/** Generates a list of selected assets in the content browser */
 	void GetSelectedAssets(TArray<UObject*>& Assets, bool SkipRedirectors) const;
 
+	/** Generates a list of selected assets in the content browser, and returns the asset data so you do not have to load them */
+	void GetSelectedAssetData(TArray<FAssetData>& AssetDataList, bool SkipRedirectors) const;
+
 private:
 	TArray<FAssetData> SelectedAssets;
 
