AssetDefinition - Converting the Haptic AssetTypeActions to AssetDefinitions.

The copying localization ID is now an option stated as such instead of a right click showing the ID you click.  In order to show the ID the asset has to be loaded, so instead it's now a copy operation that then loads the asset before copying the ID.

Fixing a regression that stopped showing export options for unloaded assets.  The code was using FastGetAsset, which won't load the asset if unloaded.  The new code just looks at the package flags from the FAssetData, no need to attempt to load the asset.

diff --git a/Engine/Plugins/Editor/ContentBrowser/ContentBrowserAssetDataSource/Source/ContentBrowserAssetDataSource/Private/AssetFileContextMenu.cpp b/Engine/Plugins/Editor/ContentBrowser/ContentBrowserAssetDataSource/Source/ContentBrowserAssetDataSource/Private/AssetFileContextMenu.cpp
--- a/Engine/Plugins/Editor/ContentBrowser/ContentBrowserAssetDataSource/Source/ContentBrowserAssetDataSource/Private/AssetFileContextMenu.cpp
+++ b/Engine/Plugins/Editor/ContentBrowser/ContentBrowserAssetDataSource/Source/ContentBrowserAssetDataSource/Private/AssetFileContextMenu.cpp
@@ -87,6 +87,8 @@
 #include "Editor/UnrealEdEngine.h"
 #include "Preferences/UnrealEdOptions.h"
 #include "UnrealEdGlobals.h"
+#include "Algo/AnyOf.h"
+#include "Misc/WarnIfAssetsLoadedInScope.h"
 
 #define LOCTEXT_NAMESPACE "ContentBrowser"
 
@@ -314,6 +316,8 @@ bool FAssetFileContextMenu::AddCommonMenuOptions(UToolMenu* Menu)
 
 void FAssetFileContextMenu::MakeAssetActionsSubMenu(UToolMenu* Menu)
 {
+	FWarnIfAssetsLoadedInScope WarnIfAssetsLoaded;
+	
 	UContentBrowserDataMenuContext_FileMenu* Context = Menu->FindContext<UContentBrowserDataMenuContext_FileMenu>();
 	const bool bCanBeModified = !Context || Context->bCanBeModified;
 
@@ -384,22 +388,12 @@ void FAssetFileContextMenu::MakeAssetActionsSubMenu(UToolMenu* Menu)
 	if (bCanBeModified)
 	{
 		FToolMenuSection& Section = Menu->AddSection("AssetContextMoveActions", LOCTEXT("AssetContextMoveActionsMenuHeading", "Move"));
-		bool bHasExportableAssets = false;
-		for (const FAssetData& AssetData : SelectedAssets)
+		const bool bCanExportAllSelectedAssets = !Algo::AnyOf(SelectedAssets, [](const FAssetData& AssetData)
 		{
-			const UObject* Object = AssetData.FastGetAsset();
-			if (Object)
-			{
-				const UPackage* Package = Object->GetOutermost();
-				if (!Package->HasAnyPackageFlags(EPackageFlags::PKG_DisallowExport))
-				{
-					bHasExportableAssets = true;
-					break;
-				}
-			}
-		}
+			return AssetData.HasAnyPackageFlags(EPackageFlags::PKG_DisallowExport);
+		});
 
-		if (bHasExportableAssets)
+		if (bCanExportAllSelectedAssets)
 		{
 			// Export
 			Section.AddMenuEntry(
@@ -658,6 +652,8 @@ bool FAssetFileContextMenu::CanExecuteAssetActions() const
 
 void FAssetFileContextMenu::MakeAssetLocalizationSubMenu(UToolMenu* Menu)
 {
+	FWarnIfAssetsLoadedInScope WarnIfAssetsLoaded;
+	
 	TArray<FCultureRef> CurrentCultures;
 
 	// Build up the list of cultures already used
@@ -784,14 +780,18 @@ void FAssetFileContextMenu::MakeAssetLocalizationSubMenu(UToolMenu* Menu)
 			// Show the localization ID if we have a single asset selected
 			if (SelectedAssets.Num() == 1)
 			{
-				const FString LocalizationId = TextNamespaceUtil::GetPackageNamespace(SelectedAssets[0].FastGetAsset());
+				FAssetData LocalizationAsset = SelectedAssets[0]; 
 				Section.AddMenuEntry(
 					"CopyLocalizationId",
-					FText::Format(LOCTEXT("CopyLocalizationIdFmt", "ID: {0}"), LocalizationId.IsEmpty() ? LOCTEXT("EmptyLocalizationId", "None") : FText::FromString(LocalizationId)),
-					LOCTEXT("CopyLocalizationIdTooltip", "Copy the localization ID to the clipboard."),
+					LOCTEXT("CopyLocalizationId", "Copy Localization ID"),
+					LOCTEXT("CopyLocalizationIdTooltip", "Load the asset and copy the localization ID to the clipboard."),
 					FSlateIcon(),
-					FUIAction(FExecuteAction::CreateSP(this, &FAssetFileContextMenu::ExecuteCopyTextToClipboard, LocalizationId))
-					);
+					FUIAction(FExecuteAction::CreateLambda([LocalizationAsset]()
+					{
+						const FString LocalizationId = TextNamespaceUtil::GetPackageNamespace(LocalizationAsset.FastGetAsset(true));
+						FPlatformApplicationMisc::ClipboardCopy(*LocalizationId);
+					}))
+				);
 			}
 
 			// Always show the reset localization ID option
@@ -1753,11 +1753,6 @@ void FAssetFileContextMenu::ExecuteGoToDocsForAsset(UClass* SelectedClass, const
 	}
 }
 
-void FAssetFileContextMenu::ExecuteCopyTextToClipboard(FString InText)
-{
-	FPlatformApplicationMisc::ClipboardCopy(*InText);
-}
-
 void FAssetFileContextMenu::ExecuteResetLocalizationId()
 {
 #if USE_STABLE_LOCALIZATION_KEYS
@@ -1891,7 +1886,7 @@ bool FAssetFileContextMenu::CanExecuteCreateBlueprintUsing() const
 	// Only work if you have a single asset selected
 	if(SelectedAssets.Num() == 1)
 	{
-		if (UObject* Asset = SelectedAssets[0].FastGetAsset())
+		if (UObject* Asset = SelectedAssets[0].FastGetAsset(true))
 		{
 			// See if we know how to make a component from this asset
 			TArray< TSubclassOf<UActorComponent> > ComponentClassList = FComponentAssetBrokerage::GetComponentsForAsset(Asset);
diff --git a/Engine/Plugins/Editor/ContentBrowser/ContentBrowserAssetDataSource/Source/ContentBrowserAssetDataSource/Public/AssetFileContextMenu.h b/Engine/Plugins/Editor/ContentBrowser/ContentBrowserAssetDataSource/Source/ContentBrowserAssetDataSource/Public/AssetFileContextMenu.h
--- a/Engine/Plugins/Editor/ContentBrowser/ContentBrowserAssetDataSource/Source/ContentBrowserAssetDataSource/Public/AssetFileContextMenu.h
+++ b/Engine/Plugins/Editor/ContentBrowser/ContentBrowserAssetDataSource/Source/ContentBrowserAssetDataSource/Public/AssetFileContextMenu.h
@@ -155,10 +155,7 @@ private:
 
 	/** Handler for GoToAssetDocs */
 	void ExecuteGoToDocsForAsset(UClass* SelectedClass, const FString ExcerptSection);
-
-	/** Handler to copy the given text to the clipboard */
-	void ExecuteCopyTextToClipboard(FString InText);
-
+	
 	/** Handler for resetting the localization ID of the current selection */
 	void ExecuteResetLocalizationId();
 
diff --git a/Engine/Plugins/Editor/EngineAssetDefinitions/Source/Private/AssetDefinition_HapticFeedbackEffectBuffer.h b/Engine/Plugins/Editor/EngineAssetDefinitions/Source/Private/AssetDefinition_HapticFeedbackEffectBuffer.h
new file mode 100644
--- /dev/null
+++ b/Engine/Plugins/Editor/EngineAssetDefinitions/Source/Private/AssetDefinition_HapticFeedbackEffectBuffer.h
@@ -0,0 +1,22 @@
+// Copyright Epic Games, Inc. All Rights Reserved.
+
+#pragma once
+
+#include "CoreMinimal.h"
+#include "Haptics/HapticFeedbackEffect_Buffer.h"
+#include "AssetDefinitionDefault.h"
+
+#include "AssetDefinition_HapticFeedbackEffectBuffer.generated.h"
+
+UCLASS()
+class UAssetDefinition_HapticFeedbackEffectBuffer : public UAssetDefinitionDefault
+{
+	GENERATED_BODY()
+
+public:
+	// UAssetDefinition Begin
+	virtual FText GetAssetDisplayName() const override { return NSLOCTEXT("AssetTypeActions", "FAssetTypeActions_HapticFeedbackEffectBuffer", "Haptic Feedback Buffer"); }
+	virtual FLinearColor GetAssetColor() const override { return FLinearColor(FColor(175, 0, 0)); }
+	virtual TSoftClassPtr<UObject> GetAssetClass() const override { return UHapticFeedbackEffect_Buffer::StaticClass(); }
+	// UAssetDefinition End
+};
diff --git a/Engine/Plugins/Editor/EngineAssetDefinitions/Source/Private/AssetDefinition_HapticFeedbackEffectCurve.h b/Engine/Plugins/Editor/EngineAssetDefinitions/Source/Private/AssetDefinition_HapticFeedbackEffectCurve.h
new file mode 100644
--- /dev/null
+++ b/Engine/Plugins/Editor/EngineAssetDefinitions/Source/Private/AssetDefinition_HapticFeedbackEffectCurve.h
@@ -0,0 +1,22 @@
+// Copyright Epic Games, Inc. All Rights Reserved.
+
+#pragma once
+
+#include "CoreMinimal.h"
+#include "Haptics/HapticFeedbackEffect_Curve.h"
+#include "AssetDefinitionDefault.h"
+
+#include "AssetDefinition_HapticFeedbackEffectCurve.generated.h"
+
+UCLASS()
+class UAssetDefinition_HapticFeedbackEffectCurve : public UAssetDefinitionDefault
+{
+	GENERATED_BODY()
+
+public:
+	// UAssetDefinition Begin
+	virtual FText GetAssetDisplayName() const override { return NSLOCTEXT("AssetTypeActions", "AssetTypeActions_HapticFeedbackEffectCurve", "Haptic Feedback Curve"); }
+	virtual FLinearColor GetAssetColor() const override { return FLinearColor(FColor(175, 0, 0)); }
+	virtual TSoftClassPtr<UObject> GetAssetClass() const override { return UHapticFeedbackEffect_Curve::StaticClass(); }
+	// UAssetDefinition End
+};
diff --git a/Engine/Plugins/Editor/EngineAssetDefinitions/Source/Private/AssetDefinition_HapticFeedbackEffectSoundWave.h b/Engine/Plugins/Editor/EngineAssetDefinitions/Source/Private/AssetDefinition_HapticFeedbackEffectSoundWave.h
new file mode 100644
--- /dev/null
+++ b/Engine/Plugins/Editor/EngineAssetDefinitions/Source/Private/AssetDefinition_HapticFeedbackEffectSoundWave.h
@@ -0,0 +1,22 @@
+// Copyright Epic Games, Inc. All Rights Reserved.
+
+#pragma once
+
+#include "CoreMinimal.h"
+#include "Haptics/HapticFeedbackEffect_SoundWave.h"
+#include "AssetDefinitionDefault.h"
+
+#include "AssetDefinition_HapticFeedbackEffectSoundWave.generated.h"
+
+UCLASS()
+class UAssetDefinition_HapticFeedbackEffectSoundWave : public UAssetDefinitionDefault
+{
+	GENERATED_BODY()
+
+public:
+	// UAssetDefinition Begin
+	virtual FText GetAssetDisplayName() const override { return NSLOCTEXT("AssetTypeActions", "FAssetTypeActions_HapticFeedbackEffectSoundWave", "Haptic Feedback Sound Wave"); }
+	virtual FLinearColor GetAssetColor() const override { return FLinearColor(FColor(175, 0, 0)); }
+	virtual TSoftClassPtr<UObject> GetAssetClass() const override { return UHapticFeedbackEffect_SoundWave::StaticClass(); }
+	// UAssetDefinition End
+};
diff --git a/Engine/Plugins/Editor/EngineAssetDefinitions/Source/Private/AssetDefinition_StaticMesh.h b/Engine/Plugins/Editor/EngineAssetDefinitions/Source/Private/AssetDefinition_StaticMesh.h
--- a/Engine/Plugins/Editor/EngineAssetDefinitions/Source/Private/AssetDefinition_StaticMesh.h
+++ b/Engine/Plugins/Editor/EngineAssetDefinitions/Source/Private/AssetDefinition_StaticMesh.h
@@ -24,7 +24,7 @@ public:
 		return Categories;
 	}
 	virtual bool CanImport() const override { return true; }
-	virtual UThumbnailInfo* LoadThumbnailInfo(const FAssetData& InAssetData) const;
+	virtual UThumbnailInfo* LoadThumbnailInfo(const FAssetData& InAssetData) const override;
 	virtual EAssetCommandResult OpenAssets(const FAssetOpenArgs& OpenArgs) const override;
 	// UAssetDefinition End
 };
diff --git a/Engine/Source/Developer/AssetTools/Private/AssetTools.cpp b/Engine/Source/Developer/AssetTools/Private/AssetTools.cpp
--- a/Engine/Source/Developer/AssetTools/Private/AssetTools.cpp
+++ b/Engine/Source/Developer/AssetTools/Private/AssetTools.cpp
@@ -61,7 +61,6 @@
 #include "AssetTypeActions/AssetTypeActions_Font.h"
 #include "AssetTypeActions/AssetTypeActions_FontFace.h"
 #include "AssetTypeActions/AssetTypeActions_ForceFeedbackEffect.h"
-#include "AssetTypeActions/AssetTypeActions_HapticFeedback.h"
 #include "AssetTypeActions/AssetTypeActions_ActorFoliageSettings.h"
 #include "AssetTypeActions/AssetTypeActions_LightWeightInstance.h"
 #include "AssetTypeActions/AssetTypeActions_MirrorDataTable.h"
@@ -1269,9 +1268,6 @@ UAssetToolsImpl::UAssetToolsImpl(const FObjectInitializer& ObjectInitializer)
 	RegisterAssetTypeActions(MakeShareable(new FAssetTypeActions_Skeleton));
 	RegisterAssetTypeActions(MakeShareable(new FAssetTypeActions_SlateBrush));
 	RegisterAssetTypeActions(MakeShareable(new FAssetTypeActions_SlateWidgetStyle));
-	RegisterAssetTypeActions(MakeShareable(new FAssetTypeActions_HapticFeedbackEffectBuffer));
-	RegisterAssetTypeActions(MakeShareable(new FAssetTypeActions_HapticFeedbackEffectCurve));
-	RegisterAssetTypeActions(MakeShareable(new FAssetTypeActions_HapticFeedbackEffectSoundWave));
 
 	// Note: Please don't add any more actions here!  They belong in an editor-only module that is more tightly
 	// coupled to your new system, and you should not create a dependency on your new system from AssetTools.
diff --git a/Engine/Source/Developer/AssetTools/Private/AssetTypeActions/AssetTypeActions_HapticFeedback.cpp b/Engine/Source/Developer/AssetTools/Private/AssetTypeActions/AssetTypeActions_HapticFeedback.cpp
deleted file mode 100644
diff --git a/Engine/Source/Developer/AssetTools/Private/AssetTypeActions/AssetTypeActions_HapticFeedback.h b/Engine/Source/Developer/AssetTools/Private/AssetTypeActions/AssetTypeActions_HapticFeedback.h
deleted file mode 100644
diff --git a/Engine/Source/Developer/AssetTools/Private/AssetTypeActions/AssetTypeActions_PhysicalMaterialMask.cpp b/Engine/Source/Developer/AssetTools/Private/AssetTypeActions/AssetTypeActions_PhysicalMaterialMask.cpp
--- a/Engine/Source/Developer/AssetTools/Private/AssetTypeActions/AssetTypeActions_PhysicalMaterialMask.cpp
+++ b/Engine/Source/Developer/AssetTools/Private/AssetTypeActions/AssetTypeActions_PhysicalMaterialMask.cpp
@@ -98,49 +98,4 @@ void FAssetTypeActions_PhysicalMaterialMask::ExecuteImport(TWeakObjectPtr<UPhysi
 	}
 }
 
-void FAssetTypeActions_PhysicalMaterialMask::ExecuteReimport(TArray<TWeakObjectPtr<UPhysicalMaterialMask>> SelectedMasks)
-{
-	for (TWeakObjectPtr<UPhysicalMaterialMask> SelectedMask : SelectedMasks)
-	{
-		if (UPhysicalMaterialMask* PhysMatMask = SelectedMask.Get())
-		{
-			FPhysicalMaterialMaskImport::ReimportMaskTexture(PhysMatMask);
-		}
-	}
-}
-
-void FAssetTypeActions_PhysicalMaterialMask::ExecuteReimportWithNewFile(TWeakObjectPtr<UPhysicalMaterialMask> SelectedMask)
-{
-	if (UPhysicalMaterialMask* PhysMatMask = SelectedMask.Get())
-	{
-		FPhysicalMaterialMaskImport::ReimportMaskTextureWithNewFile(PhysMatMask);
-	}
-}
-
-void FAssetTypeActions_PhysicalMaterialMask::ExecuteOpenSourceLocation(const TArray<FString> ResolvedFilePaths)
-{
-	// Open all files in the explorer
-	for (const auto& SourceFilePath : ResolvedFilePaths)
-	{
-		FPlatformProcess::ExploreFolder(*FPaths::GetPath(SourceFilePath));
-	}
-}
-
-void FAssetTypeActions_PhysicalMaterialMask::ExecuteOpenInExternalEditor(const TArray<FString> ResolvedFilePaths)
-{
-	// Open all files in their respective editor
-	for (const auto& SourceFilePath : ResolvedFilePaths)
-	{
-		FPlatformProcess::LaunchFileInDefaultExternalApplication(*SourceFilePath, NULL, ELaunchVerb::Edit);
-	}
-}
-
-void FAssetTypeActions_PhysicalMaterialMask::ExecuteDebug(TWeakObjectPtr<UPhysicalMaterialMask> SelectedMask)
-{
-	if (UPhysicalMaterialMask* PhysMatMask = SelectedMask.Get())
-	{
-		PhysMatMask->DumpMaskData();
-	}
-}
-
 #undef LOCTEXT_NAMESPACE
diff --git a/Engine/Source/Developer/AssetTools/Private/AssetTypeActions/AssetTypeActions_PhysicalMaterialMask.h b/Engine/Source/Developer/AssetTools/Private/AssetTypeActions/AssetTypeActions_PhysicalMaterialMask.h
--- a/Engine/Source/Developer/AssetTools/Private/AssetTypeActions/AssetTypeActions_PhysicalMaterialMask.h
+++ b/Engine/Source/Developer/AssetTools/Private/AssetTypeActions/AssetTypeActions_PhysicalMaterialMask.h
@@ -23,19 +23,4 @@ private:
 
 	/** Handler for Import */
 	void ExecuteImport(TWeakObjectPtr<UPhysicalMaterialMask> InSelectedMask);
-
-	/** Handler for Reimport */
-	void ExecuteReimport(TArray<TWeakObjectPtr<UPhysicalMaterialMask>> InSelectedMasks);
-
-	/** Handler for ReimportWithNewFile */
-	void ExecuteReimportWithNewFile(TWeakObjectPtr<UPhysicalMaterialMask> InSelectedMask);
-
-	/** Handler for OpenSourceLocation */
-	void ExecuteOpenSourceLocation(const TArray<FString> ResolvedFilePaths);
-
-	/** Handler for OpenInExternalEditor */
-	void ExecuteOpenInExternalEditor(const TArray<FString> ResolvedFilePaths);
-
-	/** Handler for Debug */
-	void ExecuteDebug(TWeakObjectPtr<UPhysicalMaterialMask> InSelectedMask);
 };
