Fixed crash introduced by CL 23288342
Rather than skipping it's initialization, avoid creation of the UContentBrowserAssetDataSource when running from a commandlet

diff --git a/Engine/Plugins/Editor/ContentBrowser/ContentBrowserAssetDataSource/Source/ContentBrowserAssetDataSource/Private/ContentBrowserAssetDataSource.cpp b/Engine/Plugins/Editor/ContentBrowser/ContentBrowserAssetDataSource/Source/ContentBrowserAssetDataSource/Private/ContentBrowserAssetDataSource.cpp
--- a/Engine/Plugins/Editor/ContentBrowser/ContentBrowserAssetDataSource/Source/ContentBrowserAssetDataSource/Private/ContentBrowserAssetDataSource.cpp
+++ b/Engine/Plugins/Editor/ContentBrowser/ContentBrowserAssetDataSource/Source/ContentBrowserAssetDataSource/Private/ContentBrowserAssetDataSource.cpp
@@ -40,129 +40,128 @@
 
 void UContentBrowserAssetDataSource::Initialize(const bool InAutoRegister)
 {
+	check(GIsEditor && !IsRunningCommandlet());
+
 	Super::Initialize(InAutoRegister);
 
-	if (GIsEditor && !IsRunningCommandlet())
-	{
-		AssetRegistry = &FModuleManager::LoadModuleChecked<FAssetRegistryModule>(AssetRegistryConstants::ModuleName).Get();
-		AssetRegistry->OnFileLoadProgressUpdated().AddUObject(this, &UContentBrowserAssetDataSource::OnAssetRegistryFileLoadProgress);
+	AssetRegistry = &FModuleManager::LoadModuleChecked<FAssetRegistryModule>(AssetRegistryConstants::ModuleName).Get();
+	AssetRegistry->OnFileLoadProgressUpdated().AddUObject(this, &UContentBrowserAssetDataSource::OnAssetRegistryFileLoadProgress);
 
-		{
-			static const FName NAME_AssetTools = "AssetTools";
-			AssetTools = &FModuleManager::GetModuleChecked<FAssetToolsModule>(NAME_AssetTools).Get();
-		}
+	{
+		static const FName NAME_AssetTools = "AssetTools";
+		AssetTools = &FModuleManager::GetModuleChecked<FAssetToolsModule>(NAME_AssetTools).Get();
+	}
 
-		CollectionManager = &FCollectionManagerModule::GetModule().Get();
+	CollectionManager = &FCollectionManagerModule::GetModule().Get();
 
-		// Listen for asset registry updates
-		AssetRegistry->OnAssetAdded().AddUObject(this, &UContentBrowserAssetDataSource::OnAssetAdded);
-		AssetRegistry->OnAssetRemoved().AddUObject(this, &UContentBrowserAssetDataSource::OnAssetRemoved);
-		AssetRegistry->OnAssetRenamed().AddUObject(this, &UContentBrowserAssetDataSource::OnAssetRenamed);
-		AssetRegistry->OnAssetUpdated().AddUObject(this, &UContentBrowserAssetDataSource::OnAssetUpdated);
-		AssetRegistry->OnPathAdded().AddUObject(this, &UContentBrowserAssetDataSource::OnPathAdded);
-		AssetRegistry->OnPathRemoved().AddUObject(this, &UContentBrowserAssetDataSource::OnPathRemoved);
-		AssetRegistry->OnFilesLoaded().AddUObject(this, &UContentBrowserAssetDataSource::OnScanCompleted);
+	// Listen for asset registry updates
+	AssetRegistry->OnAssetAdded().AddUObject(this, &UContentBrowserAssetDataSource::OnAssetAdded);
+	AssetRegistry->OnAssetRemoved().AddUObject(this, &UContentBrowserAssetDataSource::OnAssetRemoved);
+	AssetRegistry->OnAssetRenamed().AddUObject(this, &UContentBrowserAssetDataSource::OnAssetRenamed);
+	AssetRegistry->OnAssetUpdated().AddUObject(this, &UContentBrowserAssetDataSource::OnAssetUpdated);
+	AssetRegistry->OnPathAdded().AddUObject(this, &UContentBrowserAssetDataSource::OnPathAdded);
+	AssetRegistry->OnPathRemoved().AddUObject(this, &UContentBrowserAssetDataSource::OnPathRemoved);
+	AssetRegistry->OnFilesLoaded().AddUObject(this, &UContentBrowserAssetDataSource::OnScanCompleted);
 
-		// Listen for when assets are loaded or changed
-		FCoreUObjectDelegates::OnObjectPropertyChanged.AddUObject(this, &UContentBrowserAssetDataSource::OnObjectPropertyChanged);
+	// Listen for when assets are loaded or changed
+	FCoreUObjectDelegates::OnObjectPropertyChanged.AddUObject(this, &UContentBrowserAssetDataSource::OnObjectPropertyChanged);
 
-		// Listen for new mount roots
-		FPackageName::OnContentPathMounted().AddUObject(this, &UContentBrowserAssetDataSource::OnContentPathMounted);
-		FPackageName::OnContentPathDismounted().AddUObject(this, &UContentBrowserAssetDataSource::OnContentPathDismounted);
+	// Listen for new mount roots
+	FPackageName::OnContentPathMounted().AddUObject(this, &UContentBrowserAssetDataSource::OnContentPathMounted);
+	FPackageName::OnContentPathDismounted().AddUObject(this, &UContentBrowserAssetDataSource::OnContentPathDismounted);
 
-		// Listen for paths being forced visible
-		AssetViewUtils::OnAlwaysShowPath().AddUObject(this, &UContentBrowserAssetDataSource::OnAlwaysShowPath);
+	// Listen for paths being forced visible
+	AssetViewUtils::OnAlwaysShowPath().AddUObject(this, &UContentBrowserAssetDataSource::OnAlwaysShowPath);
 
-		// Register our ability to create assets via the legacy Content Browser API
-		ContentBrowserDataLegacyBridge::OnCreateNewAsset().BindUObject(this, &UContentBrowserAssetDataSource::OnBeginCreateAsset);
+	// Register our ability to create assets via the legacy Content Browser API
+	ContentBrowserDataLegacyBridge::OnCreateNewAsset().BindUObject(this, &UContentBrowserAssetDataSource::OnBeginCreateAsset);
 
-		// Create the asset menu instances
-		AssetFolderContextMenu = MakeShared<FAssetFolderContextMenu>();
-		AssetFileContextMenu = MakeShared<FAssetFileContextMenu>();
+	// Create the asset menu instances
+	AssetFolderContextMenu = MakeShared<FAssetFolderContextMenu>();
+	AssetFileContextMenu = MakeShared<FAssetFileContextMenu>();
 
-		// Bind the asset specific menu extensions
+	// Bind the asset specific menu extensions
+	{
+		if (UToolMenu* Menu = UToolMenus::Get()->ExtendMenu("ContentBrowser.AddNewContextMenu"))
 		{
-			if (UToolMenu* Menu = UToolMenus::Get()->ExtendMenu("ContentBrowser.AddNewContextMenu"))
-			{
-				Menu->AddDynamicSection(*FString::Printf(TEXT("DynamicSection_DataSource_%s"), *GetName()), FNewToolMenuDelegate::CreateLambda([WeakThis = TWeakObjectPtr<UContentBrowserAssetDataSource>(this)](UToolMenu* InMenu)
+			Menu->AddDynamicSection(*FString::Printf(TEXT("DynamicSection_DataSource_%s"), *GetName()), FNewToolMenuDelegate::CreateLambda([WeakThis = TWeakObjectPtr<UContentBrowserAssetDataSource>(this)](UToolMenu* InMenu)
+				{
+					if (UContentBrowserAssetDataSource* This = WeakThis.Get())
 					{
-						if (UContentBrowserAssetDataSource* This = WeakThis.Get())
-						{
-							This->PopulateAddNewContextMenu(InMenu);
-						}
-					}));
-			}
+						This->PopulateAddNewContextMenu(InMenu);
+					}
+				}));
+		}
 
 
-			if (UToolMenu* Menu = UToolMenus::Get()->ExtendMenu("ContentBrowser.ToolBar"))
-			{
-				Menu->AddDynamicSection(*FString::Printf(TEXT("DynamicSection_DataSource_%s"), *GetName()), FNewToolMenuDelegate::CreateLambda([WeakThis = TWeakObjectPtr<UContentBrowserAssetDataSource>(this)](UToolMenu* InMenu)
+		if (UToolMenu* Menu = UToolMenus::Get()->ExtendMenu("ContentBrowser.ToolBar"))
+		{
+			Menu->AddDynamicSection(*FString::Printf(TEXT("DynamicSection_DataSource_%s"), *GetName()), FNewToolMenuDelegate::CreateLambda([WeakThis = TWeakObjectPtr<UContentBrowserAssetDataSource>(this)](UToolMenu* InMenu)
+				{
+					if (UContentBrowserAssetDataSource* This = WeakThis.Get())
 					{
-						if (UContentBrowserAssetDataSource* This = WeakThis.Get())
-						{
-							This->PopulateContentBrowserToolBar(InMenu);
-						}
-					}));
-			}
+						This->PopulateContentBrowserToolBar(InMenu);
+					}
+				}));
+		}
 
 
-			if (UToolMenu* Menu = UToolMenus::Get()->ExtendMenu("ContentBrowser.FolderContextMenu"))
-			{
-				Menu->AddDynamicSection(*FString::Printf(TEXT("DynamicSection_DataSource_%s"), *GetName()), FNewToolMenuDelegate::CreateLambda([WeakThis = TWeakObjectPtr<UContentBrowserAssetDataSource>(this)](UToolMenu* InMenu)
+		if (UToolMenu* Menu = UToolMenus::Get()->ExtendMenu("ContentBrowser.FolderContextMenu"))
+		{
+			Menu->AddDynamicSection(*FString::Printf(TEXT("DynamicSection_DataSource_%s"), *GetName()), FNewToolMenuDelegate::CreateLambda([WeakThis = TWeakObjectPtr<UContentBrowserAssetDataSource>(this)](UToolMenu* InMenu)
+				{
+					if (UContentBrowserAssetDataSource* This = WeakThis.Get())
 					{
-						if (UContentBrowserAssetDataSource* This = WeakThis.Get())
-						{
-							This->PopulateAssetFolderContextMenu(InMenu);
-						}
-					}));
-			}
+						This->PopulateAssetFolderContextMenu(InMenu);
+					}
+				}));
+		}
 
-			if (UToolMenu* Menu = UToolMenus::Get()->ExtendMenu("ContentBrowser.AssetContextMenu"))
-			{
-				Menu->AddDynamicSection(*FString::Printf(TEXT("DynamicSection_DataSource_%s"), *GetName()), FNewToolMenuDelegate::CreateLambda([WeakThis = TWeakObjectPtr<UContentBrowserAssetDataSource>(this)](UToolMenu* InMenu)
+		if (UToolMenu* Menu = UToolMenus::Get()->ExtendMenu("ContentBrowser.AssetContextMenu"))
+		{
+			Menu->AddDynamicSection(*FString::Printf(TEXT("DynamicSection_DataSource_%s"), *GetName()), FNewToolMenuDelegate::CreateLambda([WeakThis = TWeakObjectPtr<UContentBrowserAssetDataSource>(this)](UToolMenu* InMenu)
+				{
+					if (UContentBrowserAssetDataSource* This = WeakThis.Get())
 					{
-						if (UContentBrowserAssetDataSource* This = WeakThis.Get())
-						{
-							This->PopulateAssetFileContextMenu(InMenu);
-						}
-					}));
-			}
+						This->PopulateAssetFileContextMenu(InMenu);
+					}
+				}));
+		}
 
-			if (UToolMenu* Menu = UToolMenus::Get()->ExtendMenu("ContentBrowser.DragDropContextMenu"))
-			{
-				Menu->AddDynamicSection(*FString::Printf(TEXT("DynamicSection_DataSource_%s"), *GetName()), FNewToolMenuDelegate::CreateLambda([WeakThis = TWeakObjectPtr<UContentBrowserAssetDataSource>(this)](UToolMenu* InMenu)
+		if (UToolMenu* Menu = UToolMenus::Get()->ExtendMenu("ContentBrowser.DragDropContextMenu"))
+		{
+			Menu->AddDynamicSection(*FString::Printf(TEXT("DynamicSection_DataSource_%s"), *GetName()), FNewToolMenuDelegate::CreateLambda([WeakThis = TWeakObjectPtr<UContentBrowserAssetDataSource>(this)](UToolMenu* InMenu)
+				{
+					if (UContentBrowserAssetDataSource* This = WeakThis.Get())
 					{
-						if (UContentBrowserAssetDataSource* This = WeakThis.Get())
-						{
-							This->PopulateDragDropContextMenu(InMenu);
-						}
-					}));
-			}
+						This->PopulateDragDropContextMenu(InMenu);
+					}
+				}));
 		}
+	}
 
-		DiscoveryStatusText = LOCTEXT("InitializingAssetDiscovery", "Initializing Asset Discovery...");
+	DiscoveryStatusText = LOCTEXT("InitializingAssetDiscovery", "Initializing Asset Discovery...");
 
-		// Populate the initial set of hidden empty folders
-		// This will be updated as the scan finds more content
-		AssetRegistry->EnumerateAllCachedPaths([this](FName InPath)
+	// Populate the initial set of hidden empty folders
+	// This will be updated as the scan finds more content
+	AssetRegistry->EnumerateAllCachedPaths([this](FName InPath)
+		{
+			if (!AssetRegistry->HasAssets(InPath, /*bRecursive*/true))
 			{
-				if (!AssetRegistry->HasAssets(InPath, /*bRecursive*/true))
-				{
-					EmptyAssetFolders.Add(InPath);
-				}
-				return true;
-			});
-		VisitedEmptyAssetFolders.Empty();
+				EmptyAssetFolders.Add(InPath);
+			}
+			return true;
+		});
+	VisitedEmptyAssetFolders.Empty();
 
-		FPackageName::QueryRootContentPaths(RootContentPaths);
+	FPackageName::QueryRootContentPaths(RootContentPaths);
 
-		BuildRootPathVirtualTree();
+	BuildRootPathVirtualTree();
 
-		// Mount roots are always visible
-		for (const FString& RootContentPath : RootContentPaths)
-		{
-			OnAlwaysShowPath(RootContentPath);
-		}
+	// Mount roots are always visible
+	for (const FString& RootContentPath : RootContentPaths)
+	{
+		OnAlwaysShowPath(RootContentPath);
 	}
 }
 
diff --git a/Engine/Plugins/Editor/ContentBrowser/ContentBrowserAssetDataSource/Source/ContentBrowserAssetDataSource/Private/ContentBrowserAssetDataSourceModule.cpp b/Engine/Plugins/Editor/ContentBrowser/ContentBrowserAssetDataSource/Source/ContentBrowserAssetDataSource/Private/ContentBrowserAssetDataSourceModule.cpp
--- a/Engine/Plugins/Editor/ContentBrowser/ContentBrowserAssetDataSource/Source/ContentBrowserAssetDataSource/Private/ContentBrowserAssetDataSourceModule.cpp
+++ b/Engine/Plugins/Editor/ContentBrowser/ContentBrowserAssetDataSource/Source/ContentBrowserAssetDataSource/Private/ContentBrowserAssetDataSourceModule.cpp
@@ -9,8 +9,11 @@ class FContentBrowserAssetDataSourceModule : public FDefaultModuleImpl
 public:
 	virtual void StartupModule() override
 	{
-		AssetDataSource.Reset(NewObject<UContentBrowserAssetDataSource>(GetTransientPackage(), "AssetData"));
-		AssetDataSource->Initialize();
+		if (GIsEditor && !IsRunningCommandlet())
+		{
+			AssetDataSource.Reset(NewObject<UContentBrowserAssetDataSource>(GetTransientPackage(), "AssetData"));
+			AssetDataSource->Initialize();
+		}
 	}
 
 	virtual void ShutdownModule() override
