Introducing the unsupported asset item for the content browser.
These are currently used to display the asset that are hidden by a backend class permission system.
They can also be restricted to some path only if desired.
Note: We will do a UX pass for the unsupported items after this is submitted and later we will also add some support for the unsupported asset to the alias data source.

diff --git a/Engine/Plugins/Editor/ContentBrowser/ContentBrowserAssetDataSource/Source/ContentBrowserAssetDataSource/Private/ContentBrowserAssetDataCore.cpp b/Engine/Plugins/Editor/ContentBrowser/ContentBrowserAssetDataSource/Source/ContentBrowserAssetDataSource/Private/ContentBrowserAssetDataCore.cpp
--- a/Engine/Plugins/Editor/ContentBrowser/ContentBrowserAssetDataSource/Source/ContentBrowserAssetDataSource/Private/ContentBrowserAssetDataCore.cpp
+++ b/Engine/Plugins/Editor/ContentBrowser/ContentBrowserAssetDataSource/Source/ContentBrowserAssetDataSource/Private/ContentBrowserAssetDataCore.cpp
@@ -47,9 +47,15 @@ FContentBrowserItemData CreateAssetFileItem(UContentBrowserDataSource* InOwnerDa
 	return FContentBrowserItemData(InOwnerDataSource, EContentBrowserItemFlags::Type_File | EContentBrowserItemFlags::Category_Asset, InVirtualPath, InAssetData.AssetName, FText(), MakeShared<FContentBrowserAssetFileItemDataPayload>(InAssetData));
 }
 
+FContentBrowserItemData CreateUnsupportedAssetFileItem(UContentBrowserDataSource* InOwnerDataSource, const FName InVirtualPath, const FAssetData& InAssetData)
+{
+	return FContentBrowserItemData(InOwnerDataSource, EContentBrowserItemFlags::Type_File | EContentBrowserItemFlags::Category_Asset | EContentBrowserItemFlags::Misc_Unsupported, InVirtualPath, InAssetData.AssetName, FText(), MakeShared<FContentBrowserUnsupportedAssetFileItemDataPayload>(InAssetData));
+}
+
+
 TSharedPtr<const FContentBrowserAssetFolderItemDataPayload> GetAssetFolderItemPayload(const UContentBrowserDataSource* InOwnerDataSource, const FContentBrowserItemData& InItem)
 {
-	if (InItem.GetOwnerDataSource() == InOwnerDataSource && InItem.IsFolder())
+	if (InItem.GetOwnerDataSource() == InOwnerDataSource && InItem.IsFolder() && InItem.IsSupported())
 	{
 		return StaticCastSharedPtr<const FContentBrowserAssetFolderItemDataPayload>(InItem.GetPayload());
 	}
@@ -58,13 +64,23 @@ TSharedPtr<const FContentBrowserAssetFolderItemDataPayload> GetAssetFolderItemPa
 
 TSharedPtr<const FContentBrowserAssetFileItemDataPayload> GetAssetFileItemPayload(const UContentBrowserDataSource* InOwnerDataSource, const FContentBrowserItemData& InItem)
 {
-	if (InItem.GetOwnerDataSource() == InOwnerDataSource && InItem.IsFile())
+	if (InItem.GetOwnerDataSource() == InOwnerDataSource && InItem.IsFile() && InItem.IsSupported())
 	{
 		return StaticCastSharedPtr<const FContentBrowserAssetFileItemDataPayload>(InItem.GetPayload());
 	}
 	return nullptr;
 }
 
+TSharedPtr<const FContentBrowserUnsupportedAssetFileItemDataPayload> GetUnsupportedAssetFileItemPayload(const UContentBrowserDataSource* InOwnerDataSource, const FContentBrowserItemData& InItem)
+{
+	if (InItem.GetOwnerDataSource() == InOwnerDataSource && EnumHasAllFlags(InItem.GetItemFlags(), EContentBrowserItemFlags::Type_File | EContentBrowserItemFlags::Misc_Unsupported))
+	{
+		return StaticCastSharedPtr<const FContentBrowserUnsupportedAssetFileItemDataPayload>(InItem.GetPayload());
+	}
+
+	return nullptr;
+}
+
 void EnumerateAssetFolderItemPayloads(const UContentBrowserDataSource* InOwnerDataSource, TArrayView<const FContentBrowserItemData> InItems, TFunctionRef<bool(const TSharedRef<const FContentBrowserAssetFolderItemDataPayload>&)> InFolderPayloadCallback)
 {
 	for (const FContentBrowserItemData& Item : InItems)
@@ -94,6 +110,11 @@ void EnumerateAssetFileItemPayloads(const UContentBrowserDataSource* InOwnerData
 }
 
 void EnumerateAssetItemPayloads(const UContentBrowserDataSource* InOwnerDataSource, TArrayView<const FContentBrowserItemData> InItems, TFunctionRef<bool(const TSharedRef<const FContentBrowserAssetFolderItemDataPayload>&)> InFolderPayloadCallback, TFunctionRef<bool(const TSharedRef<const FContentBrowserAssetFileItemDataPayload>&)> InAssetPayloadCallback)
+{
+	EnumerateAssetItemPayloads(InOwnerDataSource, InItems, InFolderPayloadCallback, InAssetPayloadCallback, [](const TSharedRef<const FContentBrowserUnsupportedAssetFileItemDataPayload>&){ return true; });
+}
+
+void EnumerateAssetItemPayloads(const UContentBrowserDataSource* InOwnerDataSource, TArrayView<const FContentBrowserItemData> InItems, TFunctionRef<bool(const TSharedRef<const FContentBrowserAssetFolderItemDataPayload>&)> InFolderPayloadCallback, TFunctionRef<bool(const TSharedRef<const FContentBrowserAssetFileItemDataPayload>&)> InAssetPayloadCallback, TFunctionRef<bool(const TSharedRef<const FContentBrowserUnsupportedAssetFileItemDataPayload>&)> InUnsupportedAssetPayloadCallback)
 {
 	for (const FContentBrowserItemData& Item : InItems)
 	{
@@ -112,6 +133,14 @@ void EnumerateAssetItemPayloads(const UContentBrowserDataSource* InOwnerDataSour
 				break;
 			}
 		}
+
+		if (TSharedPtr<const FContentBrowserUnsupportedAssetFileItemDataPayload> UnsupportedAssetPayload = GetUnsupportedAssetFileItemPayload(InOwnerDataSource, Item))
+		{
+			if (!InUnsupportedAssetPayloadCallback(UnsupportedAssetPayload.ToSharedRef()))
+			{
+				break;
+			}
+		}
 	}
 }
 
@@ -613,6 +642,11 @@ bool CanDeleteItem(IAssetTools* InAssetTools, IAssetRegistry* InAssetRegistry, c
 		return CanDeleteAssetFileItem(InAssetTools, *AssetPayload, OutErrorMsg);
 	}
 
+	if (TSharedPtr<const FContentBrowserUnsupportedAssetFileItemDataPayload> UnsupportedAssetPayload = GetUnsupportedAssetFileItemPayload(InOwnerDataSource, InItem))
+	{
+		return CanDeleteUnsupportedAssetFileItem(InAssetTools, *UnsupportedAssetPayload, OutErrorMsg);
+	}
+
 	return false;
 }
 
@@ -665,6 +699,29 @@ bool CanDeleteAssetFileItem(IAssetTools* InAssetTools, const FContentBrowserAsse
 	return true;
 }
 
+bool CanDeleteUnsupportedAssetFileItem(IAssetTools* InAssetTools, const FContentBrowserUnsupportedAssetFileItemDataPayload& InUnsupportedAssetPayload, FText* OutErrorMsg)
+{
+	const FAssetData* AssetData = InUnsupportedAssetPayload.GetAssetDataIfAvailable();
+
+	if (!AssetData)
+	{
+		// will need to investigate this latter
+		return false;
+	}
+
+	if (!CanModifyPath(InAssetTools, AssetData->PackageName, OutErrorMsg))
+	{
+		return false;
+	}
+
+	if (IsRunningPIE(OutErrorMsg))
+	{
+		return false;
+	}
+
+	return true;
+}
+
 bool CanPrivatizeItem(IAssetTools* InAssetTools, IAssetRegistry* InAssetRegistry, const UContentBrowserDataSource* InOwnerDataSource, const FContentBrowserItemData& InItem, FText* OutErrorMsg)
 {
 	if (TSharedPtr<const FContentBrowserAssetFolderItemDataPayload> FolderPayload = GetAssetFolderItemPayload(InOwnerDataSource, InItem))
@@ -704,7 +761,8 @@ bool CanPrivatizeAssetFileItem(IAssetTools* InAssetTools, const FContentBrowserA
 bool DeleteItems(IAssetTools* InAssetTools, IAssetRegistry* InAssetRegistry, const UContentBrowserDataSource* InOwnerDataSource, TArrayView<const FContentBrowserItemData> InItems)
 {
 	TArray<TSharedRef<const FContentBrowserAssetFolderItemDataPayload>, TInlineAllocator<16>> FolderPayloads;
-	TArray<TSharedRef<const FContentBrowserAssetFileItemDataPayload>, TInlineAllocator<16>> AssetPayloads;
+	TArray<FAssetData> AssetDatas;
+	AssetDatas.Reserve(16);
 
 	auto ProcessAssetFolderItem = [InAssetTools, InAssetRegistry, &FolderPayloads](const TSharedRef<const FContentBrowserAssetFolderItemDataPayload>& InFolderPayload)
 	{
@@ -715,16 +773,28 @@ bool DeleteItems(IAssetTools* InAssetTools, IAssetRegistry* InAssetRegistry, con
 		return true;
 	};
 
-	auto ProcessAssetFileItem = [InAssetTools, &AssetPayloads](const TSharedRef<const FContentBrowserAssetFileItemDataPayload>& InAssetPayload)
+	auto ProcessAssetFileItem = [InAssetTools, &AssetDatas](const TSharedRef<const FContentBrowserAssetFileItemDataPayload>& InAssetPayload)
 	{
 		if (CanDeleteAssetFileItem(InAssetTools, *InAssetPayload, nullptr))
 		{
-			AssetPayloads.Add(InAssetPayload);
+			AssetDatas.Add(InAssetPayload->GetAssetData());
 		}
 		return true;
 	};
 
-	EnumerateAssetItemPayloads(InOwnerDataSource, InItems, ProcessAssetFolderItem, ProcessAssetFileItem);
+	auto ProcessUnsupportedAssetFileItem = [InAssetTools, &AssetDatas](const TSharedRef<const FContentBrowserUnsupportedAssetFileItemDataPayload>& InUnsupportedAssetPayload)
+	{
+		if (CanDeleteUnsupportedAssetFileItem(InAssetTools, *InUnsupportedAssetPayload, nullptr))
+		{
+			if (const FAssetData* AssetData = InUnsupportedAssetPayload->GetAssetDataIfAvailable())
+			{
+				AssetDatas.Add(*AssetData);
+			}
+		}
+		return true;
+	};
+
+	EnumerateAssetItemPayloads(InOwnerDataSource, InItems, ProcessAssetFolderItem, ProcessAssetFileItem, ProcessUnsupportedAssetFileItem);
 
 	bool bDidDelete = false;
 
@@ -733,9 +803,9 @@ bool DeleteItems(IAssetTools* InAssetTools, IAssetRegistry* InAssetRegistry, con
 		bDidDelete |= DeleteAssetFolderItems(FolderPayloads);
 	}
 
-	if (AssetPayloads.Num() > 0)
+	if (AssetDatas.Num() > 0)
 	{
-		bDidDelete |= DeleteAssetFileItems(AssetPayloads);
+		bDidDelete |= ObjectTools::DeleteAssets(AssetDatas) > 0;
 	}
 
 	return bDidDelete;
@@ -802,6 +872,22 @@ bool DeleteAssetFileItems(TArrayView<const TSharedRef<const FContentBrowserAsset
 		&& ObjectTools::DeleteAssets(AssetsToDelete) > 0;
 }
 
+bool DeleteUnsupportedAssetFileItems(TArrayView<const TSharedRef<const FContentBrowserUnsupportedAssetFileItemDataPayload>> InUnsupportedAssetPayloads)
+{
+	TArray<FAssetData> AssetsToDelete;
+
+	for (const TSharedRef<const FContentBrowserUnsupportedAssetFileItemDataPayload>& UnsupportedAssetPayload : InUnsupportedAssetPayloads)
+	{
+		if (const FAssetData* AssetData = UnsupportedAssetPayload->GetAssetDataIfAvailable())
+		{
+			AssetsToDelete.Add(*AssetData);
+		}
+	}
+
+	return AssetsToDelete.Num() > 0
+		&& ObjectTools::DeleteAssets(AssetsToDelete) > 0;
+}
+
 bool CanRenameItem(IAssetTools* InAssetTools, const UContentBrowserDataSource* InOwnerDataSource, const FContentBrowserItemData& InItem, const FString* InNewName, FText* OutErrorMsg)
 {
 	if (TSharedPtr<const FContentBrowserAssetFolderItemDataPayload> FolderPayload = GetAssetFolderItemPayload(InOwnerDataSource, InItem))
@@ -1165,6 +1251,11 @@ bool AppendItemReference(IAssetRegistry* InAssetRegistry, const UContentBrowserD
 		return AppendAssetFileItemReference(*AssetPayload, InOutStr);
 	}
 
+	if (TSharedPtr<const FContentBrowserUnsupportedAssetFileItemDataPayload> UnsupportedAssetPayload = GetUnsupportedAssetFileItemPayload(InOwnerDataSource, InItem))
+	{
+		return AppendUnsupportedAssetFileItemReference(*UnsupportedAssetPayload, InOutStr);
+	}
+
 	return false;
 }
 
@@ -1207,6 +1298,17 @@ bool AppendAssetFileItemReference(const FContentBrowserAssetFileItemDataPayload&
 	return true;
 }
 
+bool AppendUnsupportedAssetFileItemReference(const FContentBrowserUnsupportedAssetFileItemDataPayload& InUnsupportedAssetPayload, FString& InOutStr)
+{
+	if (const FAssetData* AssetData = InUnsupportedAssetPayload.GetAssetDataIfAvailable())
+	{
+		AppendAssetExportText(*AssetData, InOutStr);
+		return true;
+	}
+
+	return false;
+}
+
 bool GetItemPhysicalPath(const UContentBrowserDataSource* InOwnerDataSource, const FContentBrowserItemData& InItem, FString& OutDiskPath)
 {
 	if (TSharedPtr<const FContentBrowserAssetFolderItemDataPayload> FolderPayload = GetAssetFolderItemPayload(InOwnerDataSource, InItem))
@@ -1219,6 +1321,11 @@ bool GetItemPhysicalPath(const UContentBrowserDataSource* InOwnerDataSource, con
 		return GetAssetFileItemPhysicalPath(*AssetPayload, OutDiskPath);
 	}
 
+	if (TSharedPtr<const FContentBrowserUnsupportedAssetFileItemDataPayload> UnsupportedAssetPayload = GetUnsupportedAssetFileItemPayload(InOwnerDataSource, InItem))
+	{
+		return GetUnsupportedAssetFileItemPhysicalPath(*UnsupportedAssetPayload, OutDiskPath);
+	}
+
 	return false;
 }
 
@@ -1246,6 +1353,18 @@ bool GetAssetFileItemPhysicalPath(const FContentBrowserAssetFileItemDataPayload&
 	return false;
 }
 
+bool GetUnsupportedAssetFileItemPhysicalPath(const FContentBrowserUnsupportedAssetFileItemDataPayload& InUnsupportedAssetPayload, FString& OutDiskPath)
+{
+	const FString& AssetFilename = InUnsupportedAssetPayload.GetFilename();
+	if (!AssetFilename.IsEmpty())
+	{
+		OutDiskPath = AssetFilename;
+		return true;
+	}
+
+	return false;
+}
+
 void GetClassItemAttribute(const FAssetData& InAssetData, const bool InIncludeMetaData, FContentBrowserItemDataAttributeValue& OutAttributeValue)
 {
 	check(InAssetData.IsValid());
@@ -1368,6 +1487,11 @@ bool GetItemAttribute(const UContentBrowserDataSource* InOwnerDataSource, const
 		return GetAssetFileItemAttribute(*AssetPayload, InIncludeMetaData, InAttributeKey, OutAttributeValue);
 	}
 
+	if (TSharedPtr<const FContentBrowserUnsupportedAssetFileItemDataPayload> UnsupportedAssetPayload = GetUnsupportedAssetFileItemPayload(InOwnerDataSource, InItem))
+	{
+		return GetUnsupportedAssetFileItemAttribute(*UnsupportedAssetPayload, InIncludeMetaData, InAttributeKey, OutAttributeValue);
+	}
+
 	return false;
 }
 
@@ -1416,16 +1540,13 @@ bool GetAssetFolderItemAttribute(const FContentBrowserAssetFolderItemDataPayload
 
 bool GetAssetFileItemAttribute(const FContentBrowserAssetFileItemDataPayload& InAssetPayload, const bool InIncludeMetaData, const FName InAttributeKey, FContentBrowserItemDataAttributeValue& OutAttributeValue)
 {
-	// Hard-coded attribute keys
+	if (GetAssetDataAttribute(InAssetPayload.GetAssetData(), InIncludeMetaData, InAttributeKey, OutAttributeValue))
 	{
-		static const FName NAME_Type = "Type";
-
-		if (InAttributeKey == ContentBrowserItemAttributes::ItemTypeName || InAttributeKey == NAME_Class || InAttributeKey == NAME_Type)
-		{
-			GetClassItemAttribute(InAssetPayload.GetAssetData(), InIncludeMetaData, OutAttributeValue);
-			return true;
-		}
+		return true;
+	}
 
+	// Hard-coded attribute keys
+	{
 		if (InAttributeKey == ContentBrowserItemAttributes::ItemTypeDisplayName)
 		{
 			if (TSharedPtr<IAssetTypeActions> AssetTypeActions = InAssetPayload.GetAssetTypeActions())
@@ -1460,77 +1581,104 @@ bool GetAssetFileItemAttribute(const FContentBrowserAssetFileItemDataPayload& In
 			return false;
 		}
 
+		if (InAttributeKey == ContentBrowserItemAttributes::ItemColor)
+		{
+			if (TSharedPtr<IAssetTypeActions> AssetTypeActions = InAssetPayload.GetAssetTypeActions())
+			{
+				const FLinearColor AssetColor = AssetTypeActions->GetTypeColor();
+				OutAttributeValue.SetValue(AssetColor.ToString());
+				return true;
+			}
+			return false;
+		}
+	}
+
+
+	return false;
+}
+
+bool GetUnsupportedAssetFileItemAttribute(const FContentBrowserUnsupportedAssetFileItemDataPayload& InUnsupportedAssetPayload, const bool InIncludeMetaData, const FName InAttributeKey, FContentBrowserItemDataAttributeValue& OutAttributeValue)
+{
+	if (const FAssetData* AssetData = InUnsupportedAssetPayload.GetAssetDataIfAvailable())
+	{
+		return GetAssetDataAttribute(*AssetData, InIncludeMetaData, InAttributeKey, OutAttributeValue);
+	}
+
+	return false;
+}
+
+bool GetAssetDataAttribute(const FAssetData& InAssetData, const bool InIncludeMetaData, const FName InAttributeKey, FContentBrowserItemDataAttributeValue& OutAttributeValue)
+{
+	// Hard-coded attribute keys
+	{
+		static const FName NAME_Type = "Type";
+
+		if (InAttributeKey == ContentBrowserItemAttributes::ItemTypeName || InAttributeKey == NAME_Class || InAttributeKey == NAME_Type)
+		{
+			GetClassItemAttribute(InAssetData, InIncludeMetaData, OutAttributeValue);
+			return true;
+		}
+
+
 		if (InAttributeKey == ContentBrowserItemAttributes::ItemDiskSize)
 		{
-			return GetDiskSizeItemAttribute(InAssetPayload.GetAssetData(), IAssetRegistry::Get(), InIncludeMetaData, OutAttributeValue);
+			return GetDiskSizeItemAttribute(InAssetData, IAssetRegistry::Get(), InIncludeMetaData, OutAttributeValue);
 		}
 
 		if (InAttributeKey == ContentBrowserItemAttributes::VirtualizedData)
 		{
-			return GetVirtualizationItemAttribute(InAssetPayload.GetAssetData(), IAssetRegistry::Get(), InIncludeMetaData, OutAttributeValue);
+			return GetVirtualizationItemAttribute(InAssetData, IAssetRegistry::Get(), InIncludeMetaData, OutAttributeValue);
 		}
 
 		if (InAttributeKey == ContentBrowserItemAttributes::ItemIsDeveloperContent)
 		{
-			const bool bIsDevelopersFolder = AssetViewUtils::IsDevelopersFolder(InAssetPayload.GetAssetData().PackageName.ToString());
+			const bool bIsDevelopersFolder = AssetViewUtils::IsDevelopersFolder(InAssetData.PackageName.ToString());
 			OutAttributeValue.SetValue(bIsDevelopersFolder);
 			return true;
 		}
 
 		if (InAttributeKey == ContentBrowserItemAttributes::ItemIsLocalizedContent)
 		{
-			const bool bIsLocalizedFolder = FPackageName::IsLocalizedPackage(InAssetPayload.GetAssetData().PackageName.ToString());
+			const bool bIsLocalizedFolder = FPackageName::IsLocalizedPackage(InAssetData.PackageName.ToString());
 			OutAttributeValue.SetValue(bIsLocalizedFolder);
 			return true;
 		}
 
 		if (InAttributeKey == ContentBrowserItemAttributes::ItemIsEngineContent)
 		{
-			const bool bIsEngineFolder = AssetViewUtils::IsEngineFolder(InAssetPayload.GetAssetData().PackageName.ToString(), /*bIncludePlugins*/true);
+			const bool bIsEngineFolder = AssetViewUtils::IsEngineFolder(InAssetData.PackageName.ToString(), /*bIncludePlugins*/true);
 			OutAttributeValue.SetValue(bIsEngineFolder);
 			return true;
 		}
 
 		if (InAttributeKey == ContentBrowserItemAttributes::ItemIsProjectContent)
 		{
-			const bool bIsProjectFolder = AssetViewUtils::IsProjectFolder(InAssetPayload.GetAssetData().PackageName.ToString(), /*bIncludePlugins*/true);
+			const bool bIsProjectFolder = AssetViewUtils::IsProjectFolder(InAssetData.PackageName.ToString(), /*bIncludePlugins*/true);
 			OutAttributeValue.SetValue(bIsProjectFolder);
 			return true;
 		}
 
 		if (InAttributeKey == ContentBrowserItemAttributes::ItemIsPluginContent)
 		{
-			const bool bIsPluginFolder = AssetViewUtils::IsPluginFolder(InAssetPayload.GetAssetData().PackageName.ToString());
+			const bool bIsPluginFolder = AssetViewUtils::IsPluginFolder(InAssetData.PackageName.ToString());
 			OutAttributeValue.SetValue(bIsPluginFolder);
 			return true;
 		}
-
-		if (InAttributeKey == ContentBrowserItemAttributes::ItemColor)
-		{
-			if (TSharedPtr<IAssetTypeActions> AssetTypeActions = InAssetPayload.GetAssetTypeActions())
-			{
-				const FLinearColor AssetColor = AssetTypeActions->GetTypeColor();//.ReinterpretAsLinear();
-				OutAttributeValue.SetValue(AssetColor.ToString());
-				return true;
-			}
-			return false;
-		}
 	}
 
 	// Generic attribute keys
 	{
-		const FAssetData& AssetData = InAssetPayload.GetAssetData();
-		const FAssetPropertyTagCache::FClassPropertyTagCache& ClassPropertyTagCache = FAssetPropertyTagCache::Get().GetCacheForClass(AssetData.AssetClassPath);
+		const FAssetPropertyTagCache::FClassPropertyTagCache& ClassPropertyTagCache = FAssetPropertyTagCache::Get().GetCacheForClass(InAssetData.AssetClassPath);
 
 		FName FoundAttributeKey = InAttributeKey;
-		FAssetDataTagMapSharedView::FFindTagResult FoundValue = AssetData.TagsAndValues.FindTag(FoundAttributeKey);
+		FAssetDataTagMapSharedView::FFindTagResult FoundValue = InAssetData.TagsAndValues.FindTag(FoundAttributeKey);
 		if (!FoundValue.IsSet())
 		{
 			// Check to see if the key we were given resolves as an alias
 			FoundAttributeKey = ClassPropertyTagCache.GetTagNameFromAlias(FoundAttributeKey);
 			if (!FoundAttributeKey.IsNone())
 			{
-				FoundValue = AssetData.TagsAndValues.FindTag(FoundAttributeKey);
+				FoundValue = InAssetData.TagsAndValues.FindTag(FoundAttributeKey);
 			}
 		}
 		if (FoundValue.IsSet())
@@ -1550,43 +1698,62 @@ bool GetItemAttributes(const UContentBrowserDataSource* InOwnerDataSource, const
 		return GetAssetFileItemAttributes(*AssetPayload, InIncludeMetaData, OutAttributeValues);
 	}
 
+	if (TSharedPtr<const FContentBrowserUnsupportedAssetFileItemDataPayload> UnsupportedAssetPayload = GetUnsupportedAssetFileItemPayload(InOwnerDataSource, InItem))
+	{
+		return GetUnsupportedAssetFileItemAttributes(*UnsupportedAssetPayload, InIncludeMetaData, OutAttributeValues);
+	}
+
 	return false;
 }
 
 bool GetAssetFileItemAttributes(const FContentBrowserAssetFileItemDataPayload& InAssetPayload, const bool InIncludeMetaData, FContentBrowserItemDataAttributeValues& OutAttributeValues)
+{
+	return GetAssetDataAttributes(InAssetPayload.GetAssetData(), InIncludeMetaData, OutAttributeValues);
+}
+
+bool GetUnsupportedAssetFileItemAttributes(const FContentBrowserUnsupportedAssetFileItemDataPayload& InUnsupportedAssetPayload, const bool InIncludeMetaData, FContentBrowserItemDataAttributeValues& OutAttributeValues)
+{
+	if (const FAssetData* AssetData = InUnsupportedAssetPayload.GetAssetDataIfAvailable())
+	{
+		return GetAssetDataAttributes(*AssetData, InIncludeMetaData, OutAttributeValues);
+	}
+
+	return false;
+}
+
+bool GetAssetDataAttributes(const FAssetData& InAssetData, const bool InIncludeMetaData, FContentBrowserItemDataAttributeValues& OutAttributeValues)
 {
 	// Hard-coded attribute keys
 	{
 		// Class
 		FContentBrowserItemDataAttributeValue& ClassAttributeValue = OutAttributeValues.Add(NAME_Class);
-		GetClassItemAttribute(InAssetPayload.GetAssetData(), InIncludeMetaData, ClassAttributeValue);
+		GetClassItemAttribute(InAssetData, InIncludeMetaData, ClassAttributeValue);
 
 		// Disk Size
 		FContentBrowserItemDataAttributeValue DiskSizeAttributeValue;
-		if (GetDiskSizeItemAttribute(InAssetPayload.GetAssetData(), IAssetRegistry::Get(), InIncludeMetaData, DiskSizeAttributeValue))
+		if (GetDiskSizeItemAttribute(InAssetData, IAssetRegistry::Get(), InIncludeMetaData, DiskSizeAttributeValue))
 		{
 			OutAttributeValues.Add(ContentBrowserItemAttributes::ItemDiskSize, MoveTemp(DiskSizeAttributeValue));
 		}
 
 		// Virtualized Payloads
 		FContentBrowserItemDataAttributeValue VirtualizedAttributeValue;
-		if (GetVirtualizationItemAttribute(InAssetPayload.GetAssetData(), IAssetRegistry::Get(), InIncludeMetaData, VirtualizedAttributeValue))
+		if (GetVirtualizationItemAttribute(InAssetData, IAssetRegistry::Get(), InIncludeMetaData, VirtualizedAttributeValue))
 		{
 			OutAttributeValues.Add(ContentBrowserItemAttributes::VirtualizedData, MoveTemp(VirtualizedAttributeValue));
-		}	
+		}
 	}
 
 	// Generic attribute keys
 	static const FTopLevelAssetPath BlueprintAssetClass = FTopLevelAssetPath(TEXT("/Script/Engine"), TEXT("Blueprint"));
 	static const FName ParentClassTag = FName("ParentClass");
 	{
-		const FAssetData& AssetData = InAssetPayload.GetAssetData();
-		const FAssetPropertyTagCache::FClassPropertyTagCache& ClassPropertyTagCache = FAssetPropertyTagCache::Get().GetCacheForClass(AssetData.AssetClassPath);
+		const FAssetPropertyTagCache::FClassPropertyTagCache& ClassPropertyTagCache = FAssetPropertyTagCache::Get().GetCacheForClass(InAssetData.AssetClassPath);
 		const FAssetPropertyTagCache::FClassPropertyTagCache* ParentClassPropertyTagCache = nullptr;
 
-		if (AssetData.AssetClassPath == BlueprintAssetClass)
+		if (InAssetData.AssetClassPath == BlueprintAssetClass)
 		{
-			FAssetTagValueRef ParentClassRef = AssetData.TagsAndValues.FindTag(ParentClassTag);
+			FAssetTagValueRef ParentClassRef = InAssetData.TagsAndValues.FindTag(ParentClassTag);
 			if (ParentClassRef.IsSet())
 			{
 				FString ParentClassName(ParentClassRef.AsString());
@@ -1602,8 +1769,8 @@ bool GetAssetFileItemAttributes(const FContentBrowserAssetFileItemDataPayload& I
 			}
 		}
 
-		OutAttributeValues.Reserve(OutAttributeValues.Num() + AssetData.TagsAndValues.Num());
-		for (const auto& TagAndValue : AssetData.TagsAndValues)
+		OutAttributeValues.Reserve(OutAttributeValues.Num() + InAssetData.TagsAndValues.Num());
+		for (const auto& TagAndValue : InAssetData.TagsAndValues)
 		{
 			FContentBrowserItemDataAttributeValue& GenericAttributeValue = OutAttributeValues.Add(TagAndValue.Key);
 			GetGenericItemAttribute(TagAndValue.Key, TagAndValue.Value.AsString(), ClassPropertyTagCache, InIncludeMetaData, GenericAttributeValue);
diff --git a/Engine/Plugins/Editor/ContentBrowser/ContentBrowserAssetDataSource/Source/ContentBrowserAssetDataSource/Private/ContentBrowserAssetDataPayload.cpp b/Engine/Plugins/Editor/ContentBrowser/ContentBrowserAssetDataSource/Source/ContentBrowserAssetDataSource/Private/ContentBrowserAssetDataPayload.cpp
--- a/Engine/Plugins/Editor/ContentBrowser/ContentBrowserAssetDataSource/Source/ContentBrowserAssetDataSource/Private/ContentBrowserAssetDataPayload.cpp
+++ b/Engine/Plugins/Editor/ContentBrowser/ContentBrowserAssetDataSource/Source/ContentBrowserAssetDataSource/Private/ContentBrowserAssetDataPayload.cpp
@@ -11,6 +11,23 @@
 #include "IAssetTypeActions.h"
 #include "Misc/PackageName.h"
 
+namespace UE::ContentBrowserAssetDataSource::Private
+{
+	void GetFilenameFromPackageName(const FString& PackageName, FString& OutFileName, TFunctionRef<UPackage*()> GetPackage)
+	{
+		// Get the filename by finding it on disk first
+		if (!FPackageName::DoesPackageExist(PackageName, &OutFileName))
+		{
+			if (const UPackage* Package = GetPackage())
+			{
+				// This is a package in memory that has not yet been saved. Determine the extension and convert to a filename
+				const FString* PackageExtension = Package->ContainsMap() ? &FPackageName::GetMapPackageExtension() : &FPackageName::GetAssetPackageExtension();
+				FPackageName::TryConvertLongPackageNameToFilename(PackageName, OutFileName, *PackageExtension);
+			}
+		}
+	}
+}
+
 const FString& FContentBrowserAssetFolderItemDataPayload::GetFilename() const
 {
 	if (!bHasCachedFilename)
@@ -127,18 +144,12 @@ const FString& FContentBrowserAssetFileItemDataPayload::GetFilename() const
 {
 	if (!bHasCachedFilename)
 	{
-		const FString PackageNameStr = AssetData.PackageName.ToString();
+		auto GetPackageLambda = [this]()
+			{ 
+				return GetPackage(); 
+			};
 
-		// Get the filename by finding it on disk first
-		if (!FPackageName::DoesPackageExist(PackageNameStr, &CachedFilename))
-		{
-			if (const UPackage* Package = GetPackage())
-			{
-				// This is a package in memory that has not yet been saved. Determine the extension and convert to a filename
-				const FString* PackageExtension = Package->ContainsMap() ? &FPackageName::GetMapPackageExtension() : &FPackageName::GetAssetPackageExtension();
-				FPackageName::TryConvertLongPackageNameToFilename(PackageNameStr, CachedFilename, *PackageExtension);
-			}
-		}
+		UE::ContentBrowserAssetDataSource::Private::GetFilenameFromPackageName(AssetData.PackageName.ToString(), CachedFilename, GetPackageLambda);
 
 		bHasCachedFilename = true;
 	}
@@ -190,3 +201,67 @@ FContentBrowserAssetFileItemDataPayload_Duplication::FContentBrowserAssetFileIte
 	, SourceObject(InSourceObject)
 {
 }
+
+FContentBrowserUnsupportedAssetFileItemDataPayload::FContentBrowserUnsupportedAssetFileItemDataPayload(FAssetData&& InAssetData)
+	: OptionalAssetData(MakeUnique<FAssetData>(MoveTemp(InAssetData)))
+{
+}
+
+FContentBrowserUnsupportedAssetFileItemDataPayload::FContentBrowserUnsupportedAssetFileItemDataPayload(const FAssetData& InAssetData)
+	: OptionalAssetData(MakeUnique<FAssetData>(InAssetData))
+{
+}
+
+const FAssetData* FContentBrowserUnsupportedAssetFileItemDataPayload::GetAssetDataIfAvailable() const
+{
+	return OptionalAssetData.Get();
+}
+
+const FString& FContentBrowserUnsupportedAssetFileItemDataPayload::GetFilename() const
+{
+	if (!bHasCachedFilename)
+	{
+		auto GetPackageLambda = [this]()
+		{
+			return GetPackage();
+		};
+
+		// Update this when we will show the asset that are to recent.
+		const FString PackageName = OptionalAssetData ? OptionalAssetData->PackageName.ToString() : FString();
+
+		UE::ContentBrowserAssetDataSource::Private::GetFilenameFromPackageName(PackageName, CachedFilename, GetPackageLambda);
+
+		bHasCachedFilename = true;
+	}
+	return CachedFilename;
+}
+
+UPackage* FContentBrowserUnsupportedAssetFileItemDataPayload::GetPackage() const
+{
+	if (bHasCachedPackagePtr && CachedPackagePtr.IsStale())
+	{
+		FlushCaches();
+	}
+
+	if (!bHasCachedPackagePtr || !CachedPackagePtr.IsValid())
+	{
+		if (const FAssetData* AssetData = OptionalAssetData.Get())
+		{ 
+			if (!AssetData->PackageName.IsNone())
+			{
+				CachedPackagePtr = FindObjectSafe<UPackage>(nullptr, *FNameBuilder(AssetData->PackageName), /*bExactClass*/true);
+			}
+		}
+		bHasCachedPackagePtr = true;
+	}
+	return CachedPackagePtr.Get();
+}
+
+void FContentBrowserUnsupportedAssetFileItemDataPayload::FlushCaches() const
+{
+	bHasCachedPackagePtr = false;
+	CachedPackagePtr.Reset();
+
+	bHasCachedFilename = false;
+	CachedFilename.Reset();
+}
diff --git a/Engine/Plugins/Editor/ContentBrowser/ContentBrowserAssetDataSource/Source/ContentBrowserAssetDataSource/Private/ContentBrowserAssetDataSource.cpp b/Engine/Plugins/Editor/ContentBrowser/ContentBrowserAssetDataSource/Source/ContentBrowserAssetDataSource/Private/ContentBrowserAssetDataSource.cpp
--- a/Engine/Plugins/Editor/ContentBrowser/ContentBrowserAssetDataSource/Source/ContentBrowserAssetDataSource/Private/ContentBrowserAssetDataSource.cpp
+++ b/Engine/Plugins/Editor/ContentBrowser/ContentBrowserAssetDataSource/Source/ContentBrowserAssetDataSource/Private/ContentBrowserAssetDataSource.cpp
@@ -245,6 +245,12 @@ bool UContentBrowserAssetDataSource::PopulateAssetFilterInputParams(FAssetFilter
 	Params.AssetDataFilter->ItemAttributeFilter = InFilter.ItemAttributeFilter;
 	Params.InternalPaths.Reset();
 
+	Params.UnsupportedClassFilter = InFilter.ExtraFilters.FindFilter<FContentBrowserDataUnsupportedClassFilter>();
+	if (Params.UnsupportedClassFilter && Params.UnsupportedClassFilter->ClassPermissionList && Params.UnsupportedClassFilter->ClassPermissionList->HasFiltering())
+	{
+		Params.ConvertToUnsupportedAssetDataFilter = &Params.FilterList->FindOrAddFilter<FContentBrowserCompiledUnsupportedAssetDataFilter>();
+	}
+
 	return true;
 }
 
@@ -408,7 +414,7 @@ bool UContentBrowserAssetDataSource::CreateAssetFilter(FAssetFilterInputParams&
 	}
 
 	// If we are filtering all classes, then we can bail now as we won't return any content
-	if (Params.ClassPermissionList && Params.ClassPermissionList->IsDenyListAll())
+	if (Params.ClassPermissionList && Params.ClassPermissionList->IsDenyListAll() && !Params.UnsupportedClassFilter)
 	{
 		return false;
 	}
@@ -484,6 +490,30 @@ PRAGMA_ENABLE_DEPRECATION_WARNINGS
 			CreateCompiledFilter(InclusiveFilter, CompiledInclusiveFilter);
 		}
 
+
+		// Add the backend class filtering to the unsupported asset filtering before the class permission are added
+		if (Params.ConvertToUnsupportedAssetDataFilter)
+		{
+			if (Params.UnsupportedClassFilter)
+			{
+				if (const FPathPermissionList* ClassPermissionList = Params.UnsupportedClassFilter->ClassPermissionList.Get())
+				{
+					if (ClassPermissionList->HasFiltering())
+					{
+						if (Params.ConvertToUnsupportedAssetDataFilter->InclusiveFilter.IsEmpty())
+						{
+							Params.ConvertToUnsupportedAssetDataFilter->InclusiveFilter.ClassPaths = CompiledInclusiveFilter.ClassPaths;
+						}
+						else
+						{
+							Params.ConvertToUnsupportedAssetDataFilter->InclusiveFilter.ClassPaths = Params.ConvertToUnsupportedAssetDataFilter->InclusiveFilter.ClassPaths.Intersect(CompiledInclusiveFilter.ClassPaths);
+						}
+					}
+				}
+			}
+		}
+
+
 		// Remove any inclusive paths that aren't under the set of internal paths that we want to enumerate
 		{
 			FARCompiledFilter CompiledInternalPathFilter;
@@ -525,6 +555,21 @@ PRAGMA_ENABLE_DEPRECATION_WARNINGS
 			}
 		}
 
+		// Add the backend class filtering to the unsupported asset filtering before the class permission are added
+		if (Params.ConvertToUnsupportedAssetDataFilter)
+		{
+			if (Params.UnsupportedClassFilter)
+			{
+				if (const FPathPermissionList* ClassPermissionList = Params.UnsupportedClassFilter->ClassPermissionList.Get())
+				{
+					if (ClassPermissionList->HasFiltering())
+					{
+						Params.ConvertToUnsupportedAssetDataFilter->InclusiveFilter.ClassPaths.Append(CompiledInclusiveFilter.ClassPaths);
+					}
+				}
+			}
+		}
+
 		// Remove any inclusive paths that aren't in the explicit AllowList set
 		if (Params.PathPermissionList && Params.PathPermissionList->GetAllowList().Num() > 0)
 		{
@@ -574,7 +619,7 @@ PRAGMA_ENABLE_DEPRECATION_WARNINGS
 				// Explicit classes given - remove anything not in the allow list class set
 				// If the classes resolve as empty then the combined filter will return nothing and can be skipped
 				CompiledInclusiveFilter.ClassPaths = CompiledInclusiveFilter.ClassPaths.Intersect(CompiledClassFilterAllowList.ClassPaths);
-				if (CompiledInclusiveFilter.ClassPaths.Num() == 0)
+				if (CompiledInclusiveFilter.ClassPaths.Num() == 0 && !Params.ConvertToUnsupportedAssetDataFilter)
 				{
 					return false;
 				}
@@ -682,7 +727,7 @@ PRAGMA_ENABLE_DEPRECATION_WARNINGS
 		if (CompiledInclusiveFilter.ClassPaths.Num() > 0 && CompiledExclusiveFilter.ClassPaths.Num() > 0)
 		{
 			CompiledInclusiveFilter.ClassPaths = CompiledInclusiveFilter.ClassPaths.Difference(CompiledExclusiveFilter.ClassPaths);
-			if (CompiledInclusiveFilter.ClassPaths.Num() == 0)
+			if (CompiledInclusiveFilter.ClassPaths.Num() == 0 && !!Params.ConvertToUnsupportedAssetDataFilter)
 			{
 				return false;
 			}
@@ -714,6 +759,170 @@ PRAGMA_ENABLE_DEPRECATION_WARNINGS
 	Params.AssetDataFilter->InclusiveFilter = MoveTemp(CompiledInclusiveFilter);
 	Params.AssetDataFilter->ExclusiveFilter = MoveTemp(CompiledExclusiveFilter);
 
+
+	// Compile the filter to show the unsupported items
+	if (Params.ConvertToUnsupportedAssetDataFilter)
+	{
+		if (Params.UnsupportedClassFilter)
+		{
+			if (const FPathPermissionList* ClassPermissionList = Params.UnsupportedClassFilter->ClassPermissionList.Get())
+			{
+				if (ClassPermissionList->HasFiltering())
+				{
+					// Create a Backend filter for the unsupported items
+					{
+						// Cache the existing class path
+						TSet<FTopLevelAssetPath> InclusiveClassPath = MoveTemp(Params.ConvertToUnsupportedAssetDataFilter->InclusiveFilter.ClassPaths);
+						TSet<FTopLevelAssetPath> ExclusiveClassPath = MoveTemp(Params.ConvertToUnsupportedAssetDataFilter->ExclusiveFilter.ClassPaths);
+
+						// Remove temporally the class filtering from the asset data filter
+						TSet<FTopLevelAssetPath> AssetDataInclusiveClassPath = MoveTemp(Params.AssetDataFilter->InclusiveFilter.ClassPaths);
+						TSet<FTopLevelAssetPath> AssetDataExclusiveClassPath = MoveTemp(Params.AssetDataFilter->ExclusiveFilter.ClassPaths);
+
+						Params.ConvertToUnsupportedAssetDataFilter->InclusiveFilter = Params.AssetDataFilter->InclusiveFilter;
+						Params.ConvertToUnsupportedAssetDataFilter->ExclusiveFilter = Params.AssetDataFilter->ExclusiveFilter;
+
+						// Restore the class filtering
+						Params.AssetDataFilter->InclusiveFilter.ClassPaths = MoveTemp(AssetDataInclusiveClassPath);
+						Params.AssetDataFilter->ExclusiveFilter.ClassPaths = MoveTemp(AssetDataExclusiveClassPath);
+
+						Params.ConvertToUnsupportedAssetDataFilter->InclusiveFilter.ClassPaths = MoveTemp(InclusiveClassPath);
+						Params.ConvertToUnsupportedAssetDataFilter->ExclusiveFilter.ClassPaths = MoveTemp(ExclusiveClassPath);
+					}
+
+					FPathPermissionList* FolderPermissionList = Params.UnsupportedClassFilter->FolderPermissionList.Get();
+
+					// Compile the inclusive filter for where to show the unsupported asset
+					{
+						FARCompiledFilter CompiledShowInclusiveFilter;
+
+						// Only show the unsupported asset in the specified folders
+						if (FolderPermissionList && FolderPermissionList->HasFiltering())
+						{
+							FARFilter ShowInclusiveFilter;
+
+							ShowInclusiveFilter.bRecursivePaths = true;
+
+							ShowInclusiveFilter.PackagePaths.Reserve(FolderPermissionList->GetAllowList().Num());
+							for (const auto& AllowListPair : ClassPermissionList->GetAllowList())
+							{
+								ShowInclusiveFilter.PackagePaths.Add(*(AllowListPair.Key));
+							}
+
+							Params.AssetRegistry->CompileFilter(ShowInclusiveFilter, CompiledShowInclusiveFilter);
+						}
+
+						if (CompiledShowInclusiveFilter.IsEmpty())
+						{
+							if (Params.ConvertToUnsupportedAssetDataFilter->ShowInclusiveFilter.PackagePaths.IsEmpty())
+							{
+								Params.ConvertToUnsupportedAssetDataFilter->ShowInclusiveFilter.PackagePaths = Params.ConvertToUnsupportedAssetDataFilter->InclusiveFilter.PackagePaths;
+							}
+							else
+							{ 
+								Params.ConvertToUnsupportedAssetDataFilter->ShowInclusiveFilter.PackagePaths = Params.ConvertToUnsupportedAssetDataFilter->ShowInclusiveFilter.PackagePaths.Intersect(Params.ConvertToUnsupportedAssetDataFilter->InclusiveFilter.PackagePaths);
+							}
+						}
+						else
+						{
+							CompiledInclusiveFilter.PackagePaths = CompiledShowInclusiveFilter.PackagePaths.Intersect(Params.ConvertToUnsupportedAssetDataFilter->InclusiveFilter.PackagePaths);
+
+							if (Params.ConvertToUnsupportedAssetDataFilter->ShowInclusiveFilter.PackagePaths.IsEmpty())
+							{
+								Params.ConvertToUnsupportedAssetDataFilter->ShowInclusiveFilter.PackagePaths = MoveTemp(CompiledInclusiveFilter.PackagePaths);
+							}
+							else
+							{
+								Params.ConvertToUnsupportedAssetDataFilter->ShowInclusiveFilter.PackagePaths = Params.ConvertToUnsupportedAssetDataFilter->ShowInclusiveFilter.PackagePaths.Intersect(CompiledInclusiveFilter.PackagePaths);
+							}
+						}
+					}
+
+					// Compile the exclusive filter for where to show the unsupported asset
+					{
+						FARCompiledFilter CompiledShowExclusiveFilter;
+
+						// Only show the unsupported asset in the specified folders
+						if (FolderPermissionList && FolderPermissionList->HasFiltering())
+						{
+							FARFilter ShowExclusiveFilter;
+
+							ShowExclusiveFilter.bRecursivePaths = true;
+
+							ShowExclusiveFilter.PackagePaths.Reserve(FolderPermissionList->GetDenyList().Num());
+							for (const auto& DenyListPair : ClassPermissionList->GetDenyList())
+							{
+								ShowExclusiveFilter.PackagePaths.Add(*(DenyListPair.Key));
+							}
+
+							Params.AssetRegistry->CompileFilter(ShowExclusiveFilter, CompiledShowExclusiveFilter);
+						}
+
+						CompiledShowExclusiveFilter.PackagePaths.Append(Params.ConvertToUnsupportedAssetDataFilter->ExclusiveFilter.PackagePaths);
+
+						if (Params.ConvertToUnsupportedAssetDataFilter->ShowExclusiveFilter.PackagePaths.IsEmpty())
+						{
+							Params.ConvertToUnsupportedAssetDataFilter->ShowExclusiveFilter.PackagePaths = MoveTemp(CompiledShowExclusiveFilter.PackagePaths);
+						}
+						else
+						{
+							Params.ConvertToUnsupportedAssetDataFilter->ShowExclusiveFilter.PackagePaths.Append(MoveTemp(CompiledShowExclusiveFilter.PackagePaths));
+						}
+					}
+
+
+					// Compile the convert if fail inclusive filter
+					if (!ClassPermissionList->GetAllowList().IsEmpty())
+					{
+						FARCompiledFilter CompiledConvertIfFailInclusiveFilter;
+						FARFilter ConvertIfFailInclusiveFilter;
+
+						// Remove any inclusive classes that aren't in the explicit allow list set
+						ConvertIfFailInclusiveFilter.ClassPaths.Reserve(ClassPermissionList->GetAllowList().Num());
+						for (const auto& AllowListPair : ClassPermissionList->GetAllowList())
+						{
+							ConvertIfFailInclusiveFilter.ClassPaths.Add(FTopLevelAssetPath(AllowListPair.Key));
+						}
+
+						Params.AssetRegistry->CompileFilter(ConvertIfFailInclusiveFilter, CompiledConvertIfFailInclusiveFilter);
+					
+
+						if (Params.ConvertToUnsupportedAssetDataFilter->ConvertIfFailInclusiveFilter.ClassPaths.Num() > 0)
+						{
+							// Explicit classes given - remove anything not in the allow list class set
+							// If the classes resolve as empty then the combined filter will return nothing and can be skipped
+							Params.ConvertToUnsupportedAssetDataFilter->ConvertIfFailInclusiveFilter.ClassPaths = Params.ConvertToUnsupportedAssetDataFilter->ConvertIfFailInclusiveFilter.ClassPaths.Intersect(CompiledConvertIfFailInclusiveFilter.ClassPaths);
+						}
+						else
+						{
+							// No explicit classes given - just use the allow list class set
+							Params.ConvertToUnsupportedAssetDataFilter->ConvertIfFailInclusiveFilter.ClassPaths = MoveTemp(CompiledConvertIfFailInclusiveFilter.ClassPaths);
+						}
+					}
+
+					// Compile the convert if fail exclusive filter
+					if (!ClassPermissionList->GetDenyList().IsEmpty())
+					{
+						FARCompiledFilter CompiledConvertIfFailExclusiveFilter;
+						FARFilter ConvertIfFailExclusiveFilter;
+
+						// Add any exclusive classes that are in the explicit DenyList set
+						ConvertIfFailExclusiveFilter.ClassPaths.Reserve(ClassPermissionList->GetDenyList().Num());
+						for (const auto& FilterPair : ClassPermissionList->GetDenyList())
+						{
+							ConvertIfFailExclusiveFilter.ClassPaths.Add(FTopLevelAssetPath(FilterPair.Key));
+						}
+
+
+						Params.AssetRegistry->CompileFilter(ConvertIfFailExclusiveFilter, CompiledConvertIfFailExclusiveFilter);
+
+						Params.ConvertToUnsupportedAssetDataFilter->ConvertIfFailExclusiveFilter.ClassPaths.Append(MoveTemp(CompiledConvertIfFailExclusiveFilter.ClassPaths));
+					}
+				}
+			}
+		}
+	}
+
 	return true;
 }
 
@@ -911,18 +1120,48 @@ void UContentBrowserAssetDataSource::EnumerateItemsMatchingFilter(const FContent
 			}
 		}
 
-		AssetRegistry->EnumerateAssets(AssetDataFilter->InclusiveFilter, [this, &InCallback, &AssetDataFilter](const FAssetData& AssetData)
+		if (const FContentBrowserCompiledUnsupportedAssetDataFilter* UnsupportedAssetDataFilter = FilterList->FindFilter<FContentBrowserCompiledUnsupportedAssetDataFilter>())
 		{
-			if (ContentBrowserAssetData::IsPrimaryAsset(AssetData))
+			// Using the show unsupported asset filter
+			AssetRegistry->EnumerateAssets(UnsupportedAssetDataFilter->InclusiveFilter, [this, &InCallback, &AssetDataFilter, UnsupportedAssetDataFilter](const FAssetData& AssetData)
 			{
-				const bool bPassesExclusiveFilter = AssetDataFilter->ExclusiveFilter.IsEmpty() || !AssetRegistry->IsAssetIncludedByFilter(AssetData, AssetDataFilter->ExclusiveFilter);
-				if (bPassesExclusiveFilter)
+				if (ContentBrowserAssetData::IsPrimaryAsset(AssetData) && AssetData.GetOptionalOuterPathName().IsNone())
 				{
-					return InCallback(CreateAssetFileItem(AssetData));
+					const bool bPassesExclusiveFilter = UnsupportedAssetDataFilter->ExclusiveFilter.IsEmpty() || !AssetRegistry->IsAssetIncludedByFilter(AssetData, UnsupportedAssetDataFilter->ExclusiveFilter);
+					if (bPassesExclusiveFilter)
+					{
+						// Should this asset be presented as unsupported
+						if (!(AssetRegistry->IsAssetIncludedByFilter(AssetData, UnsupportedAssetDataFilter->ConvertIfFailInclusiveFilter) && (UnsupportedAssetDataFilter->ConvertIfFailExclusiveFilter.IsEmpty() || AssetRegistry->IsAssetExcludedByFilter(AssetData, UnsupportedAssetDataFilter->ConvertIfFailExclusiveFilter))) // Do we fail the supported filter?
+							&& (AssetRegistry->IsAssetIncludedByFilter(AssetData, UnsupportedAssetDataFilter->ShowInclusiveFilter) && (UnsupportedAssetDataFilter->ShowExclusiveFilter.IsEmpty() || AssetRegistry->IsAssetExcludedByFilter(AssetData, UnsupportedAssetDataFilter->ShowExclusiveFilter)))) // Do we pass the show filter for the unsupported asset?
+						{
+							return InCallback(CreateUnsupportedAssetFileItem(AssetData));
+						}
+
+						// Normal item test it against the class filter
+						if ((AssetDataFilter->InclusiveFilter.ClassPaths.IsEmpty() || AssetDataFilter->InclusiveFilter.ClassPaths.Contains(AssetData.AssetClassPath)) && !AssetDataFilter->ExclusiveFilter.ClassPaths.Contains(AssetData.AssetClassPath))
+						{
+							return InCallback(CreateAssetFileItem(AssetData));
+						}
+					}
 				}
-			}
-			return true;
-		});
+				return true;
+			});
+		}
+		else
+		{
+			AssetRegistry->EnumerateAssets(AssetDataFilter->InclusiveFilter, [this, &InCallback, &AssetDataFilter](const FAssetData& AssetData)
+				{
+					if (ContentBrowserAssetData::IsPrimaryAsset(AssetData))
+					{
+						const bool bPassesExclusiveFilter = AssetDataFilter->ExclusiveFilter.IsEmpty() || !AssetRegistry->IsAssetIncludedByFilter(AssetData, AssetDataFilter->ExclusiveFilter);
+						if (bPassesExclusiveFilter)
+						{
+							return InCallback(CreateAssetFileItem(AssetData));
+						}
+					}
+					return true;
+				});
+		}
 	}
 }
 
@@ -1187,12 +1426,28 @@ bool UContentBrowserAssetDataSource::DoesItemPassFilter(const FContentBrowserIte
 	case EContentBrowserItemFlags::Type_File:
 		if (EnumHasAnyFlags(InFilter.ItemTypeFilter, EContentBrowserItemTypeFilter::IncludeFiles) && !AssetDataFilter->bFilterExcludesAllAssets)
 		{
+			auto FilterWithAssetData = [this, AssetDataFilter](const FAssetData& InAssetData, const FARCompiledFilter& InclusiveFilter, const FARCompiledFilter& ExclusiveFilter) -> bool
+				{
+				// Must pass Inclusive AND !Exclusive, or be a CustomAsset
+				return (AssetRegistry->IsAssetIncludedByFilter(InAssetData, InclusiveFilter) // InclusiveFilter
+					&& (ExclusiveFilter.IsEmpty() || !AssetRegistry->IsAssetIncludedByFilter(InAssetData, ExclusiveFilter))) // ExclusiveFilter
+					|| AssetDataFilter->CustomSourceAssets.Contains(InAssetData); // CustomAsset
+				};
+
 			if (TSharedPtr<const FContentBrowserAssetFileItemDataPayload> AssetPayload = GetAssetFileItemPayload(InItem))
 			{
-				// Must pass Inclusive AND !Exclusive, or be a CustomAsset
-				return ((AssetDataFilter->InclusiveFilter.IsEmpty() || AssetRegistry->IsAssetIncludedByFilter(AssetPayload->GetAssetData(), AssetDataFilter->InclusiveFilter)) // InclusiveFilter
-					&& (AssetDataFilter->ExclusiveFilter.IsEmpty() || !AssetRegistry->IsAssetIncludedByFilter(AssetPayload->GetAssetData(), AssetDataFilter->ExclusiveFilter))) // ExclusiveFilter
-					|| AssetDataFilter->CustomSourceAssets.Contains(AssetPayload->GetAssetData()); // CustomAsset
+				return FilterWithAssetData(AssetPayload->GetAssetData(), AssetDataFilter->InclusiveFilter, AssetDataFilter->ExclusiveFilter);
+			}
+
+			if (TSharedPtr<const FContentBrowserUnsupportedAssetFileItemDataPayload> UnsupportedAssetPayload = GetUnsupportedAssetFileItemPayload(InItem))
+			{
+				if (const FContentBrowserCompiledUnsupportedAssetDataFilter* UnsupportedAssetFilter = FilterList->FindFilter<FContentBrowserCompiledUnsupportedAssetDataFilter>())
+				{
+					if (const FAssetData* AssetData = UnsupportedAssetPayload->GetAssetDataIfAvailable())
+					{
+						return FilterWithAssetData(*AssetData, UnsupportedAssetFilter->InclusiveFilter, UnsupportedAssetFilter->ExclusiveFilter);
+					}
+				}
 			}
 		}
 		break;
@@ -1204,6 +1459,47 @@ bool UContentBrowserAssetDataSource::DoesItemPassFilter(const FContentBrowserIte
 	return false;
 }
 
+bool UContentBrowserAssetDataSource::ConvertItemForFilter(FContentBrowserItemData& Item, const FContentBrowserDataCompiledFilter& InFilter)
+{
+	const FContentBrowserDataFilterList* FilterList = InFilter.CompiledFilters.Find(this);
+	if (!FilterList)
+	{
+		return false;
+	}
+
+	const FContentBrowserCompiledUnsupportedAssetDataFilter* UnsupportedAssetDataFilter = FilterList->FindFilter<FContentBrowserCompiledUnsupportedAssetDataFilter>();
+	if (!UnsupportedAssetDataFilter)
+	{
+		return false;
+	}
+
+	if (Item.GetOwnerDataSource() != this)
+	{
+		return false;
+	}
+
+	const FContentBrowserCompiledAssetDataFilter* AssetDataFilter = FilterList->FindFilter<FContentBrowserCompiledAssetDataFilter>();
+
+	if (TSharedPtr<const FContentBrowserAssetFileItemDataPayload> AssetPayload = GetAssetFileItemPayload(Item))
+	{
+		const FAssetData& AssetData = AssetPayload->GetAssetData();
+		if (AssetData.GetOptionalOuterPathName().IsNone() && ContentBrowserAssetData::IsPrimaryAsset(AssetData))
+		{
+			if (!AssetDataFilter || !AssetDataFilter->CustomSourceAssets.Contains(AssetData))
+			{
+				if (!(AssetRegistry->IsAssetIncludedByFilter(AssetData, UnsupportedAssetDataFilter->ConvertIfFailInclusiveFilter) && (UnsupportedAssetDataFilter->ConvertIfFailExclusiveFilter.IsEmpty() || AssetRegistry->IsAssetExcludedByFilter(AssetData, UnsupportedAssetDataFilter->ConvertIfFailExclusiveFilter))) // Do we fail the supported filter?
+					&& (AssetRegistry->IsAssetIncludedByFilter(AssetData, UnsupportedAssetDataFilter->ShowInclusiveFilter) && (UnsupportedAssetDataFilter->ShowExclusiveFilter.IsEmpty() || AssetRegistry->IsAssetExcludedByFilter(AssetData, UnsupportedAssetDataFilter->ShowExclusiveFilter)))) // Do we pass the show filter for the unsupported asset?
+				{
+					Item = CreateUnsupportedAssetFileItem(AssetData);
+					return true;
+				}
+			}
+		}
+	}
+
+	return false;
+}
+
 bool UContentBrowserAssetDataSource::GetItemAttribute(const FContentBrowserItemData& InItem, const bool InIncludeMetaData, const FName InAttributeKey, FContentBrowserItemDataAttributeValue& OutAttributeValue)
 {
 	return ContentBrowserAssetData::GetItemAttribute(this, InItem, InIncludeMetaData, InAttributeKey, OutAttributeValue);
@@ -1410,6 +1706,12 @@ bool UContentBrowserAssetDataSource::RenameItem(const FContentBrowserItemData& I
 
 bool UContentBrowserAssetDataSource::CanCopyItem(const FContentBrowserItemData& InItem, const FName InDestPath, FText* OutErrorMsg)
 {
+	if (!InItem.IsSupported())
+	{
+		ContentBrowserAssetData::SetOptionalErrorMessage(OutErrorMsg, FText::Format(LOCTEXT("Error_AssetIsNotSupported", "Asset {0} is not supported and it can't be copied"), InItem.GetDisplayName()));
+		return false;
+	}
+
 	// Cannot copy an item outside the paths known to this data source
 	FName InternalDestPath;
 	if (!TryConvertVirtualPathToInternal(InDestPath, InternalDestPath))
@@ -1436,6 +1738,12 @@ bool UContentBrowserAssetDataSource::CanCopyItem(const FContentBrowserItemData&
 
 bool UContentBrowserAssetDataSource::CopyItem(const FContentBrowserItemData& InItem, const FName InDestPath)
 {
+	if (!InItem.IsSupported())
+	{
+		return false;
+	}
+
+
 	FName InternalDestPath;
 	if (!TryConvertVirtualPathToInternal(InDestPath, InternalDestPath))
 	{
@@ -1707,6 +2015,16 @@ PRAGMA_ENABLE_DEPRECATION_WARNINGS
 	return ContentBrowserAssetData::CreateAssetFileItem(this, VirtualizedPath, InAssetData);
 }
 
+FContentBrowserItemData UContentBrowserAssetDataSource::CreateUnsupportedAssetFileItem(const FAssetData& InAssetData)
+{
+	FName VirtualizedPath;
+	PRAGMA_DISABLE_DEPRECATION_WARNINGS
+		TryConvertInternalPathToVirtual(InAssetData.ObjectPath, VirtualizedPath);
+	PRAGMA_ENABLE_DEPRECATION_WARNINGS
+
+	return ContentBrowserAssetData::CreateUnsupportedAssetFileItem(this, VirtualizedPath, InAssetData);
+}
+
 TSharedPtr<const FContentBrowserAssetFolderItemDataPayload> UContentBrowserAssetDataSource::GetAssetFolderItemPayload(const FContentBrowserItemData& InItem) const
 {
 	return ContentBrowserAssetData::GetAssetFolderItemPayload(this, InItem);
@@ -1717,6 +2035,11 @@ TSharedPtr<const FContentBrowserAssetFileItemDataPayload> UContentBrowserAssetDa
 	return ContentBrowserAssetData::GetAssetFileItemPayload(this, InItem);
 }
 
+TSharedPtr<const FContentBrowserUnsupportedAssetFileItemDataPayload> UContentBrowserAssetDataSource::GetUnsupportedAssetFileItemPayload(const FContentBrowserItemData& InItem) const
+{
+	return ContentBrowserAssetData::GetUnsupportedAssetFileItemPayload(this, InItem);
+}
+
 void UContentBrowserAssetDataSource::OnAssetRegistryFileLoadProgress(const IAssetRegistry::FFileLoadProgressUpdateData& InProgressUpdateData)
 {
 	if (InProgressUpdateData.bIsDiscoveringAssetFiles)
diff --git a/Engine/Plugins/Editor/ContentBrowser/ContentBrowserAssetDataSource/Source/ContentBrowserAssetDataSource/Public/ContentBrowserAssetDataCore.h b/Engine/Plugins/Editor/ContentBrowser/ContentBrowserAssetDataSource/Source/ContentBrowserAssetDataSource/Public/ContentBrowserAssetDataCore.h
--- a/Engine/Plugins/Editor/ContentBrowser/ContentBrowserAssetDataSource/Source/ContentBrowserAssetDataSource/Public/ContentBrowserAssetDataCore.h
+++ b/Engine/Plugins/Editor/ContentBrowser/ContentBrowserAssetDataSource/Source/ContentBrowserAssetDataSource/Public/ContentBrowserAssetDataCore.h
@@ -6,6 +6,7 @@
 
 class FContentBrowserAssetFileItemDataPayload;
 class FContentBrowserAssetFolderItemDataPayload;
+class FContentBrowserUnsupportedAssetFileItemDataPayload;
 
 class IAssetTools;
 class IAssetRegistry;
@@ -23,16 +24,22 @@ namespace ContentBrowserAssetData
 
 	CONTENTBROWSERASSETDATASOURCE_API FContentBrowserItemData CreateAssetFileItem(UContentBrowserDataSource* InOwnerDataSource, const FName InVirtualPath, const FAssetData& InAssetData);
 
+	CONTENTBROWSERASSETDATASOURCE_API FContentBrowserItemData CreateUnsupportedAssetFileItem(UContentBrowserDataSource* InOwnerDataSource, const FName InVirtualPath, const FAssetData& InAssetData);
+
 	CONTENTBROWSERASSETDATASOURCE_API TSharedPtr<const FContentBrowserAssetFolderItemDataPayload> GetAssetFolderItemPayload(const UContentBrowserDataSource* InOwnerDataSource, const FContentBrowserItemData& InItem);
 
 	CONTENTBROWSERASSETDATASOURCE_API TSharedPtr<const FContentBrowserAssetFileItemDataPayload> GetAssetFileItemPayload(const UContentBrowserDataSource* InOwnerDataSource, const FContentBrowserItemData& InItem);
-	
+
+	CONTENTBROWSERASSETDATASOURCE_API TSharedPtr<const FContentBrowserUnsupportedAssetFileItemDataPayload> GetUnsupportedAssetFileItemPayload(const UContentBrowserDataSource* InOwnerDataSource, const FContentBrowserItemData& InItem);
+
 	CONTENTBROWSERASSETDATASOURCE_API void EnumerateAssetFolderItemPayloads(const UContentBrowserDataSource* InOwnerDataSource, TArrayView<const FContentBrowserItemData> InItems, TFunctionRef<bool(const TSharedRef<const FContentBrowserAssetFolderItemDataPayload>&)> InFolderPayloadCallback);
 
 	CONTENTBROWSERASSETDATASOURCE_API void EnumerateAssetFileItemPayloads(const UContentBrowserDataSource* InOwnerDataSource, TArrayView<const FContentBrowserItemData> InItems, TFunctionRef<bool(const TSharedRef<const FContentBrowserAssetFileItemDataPayload>&)> InAssetPayloadCallback);
 
 	CONTENTBROWSERASSETDATASOURCE_API void EnumerateAssetItemPayloads(const UContentBrowserDataSource* InOwnerDataSource, TArrayView<const FContentBrowserItemData> InItems, TFunctionRef<bool(const TSharedRef<const FContentBrowserAssetFolderItemDataPayload>&)> InFolderPayloadCallback, TFunctionRef<bool(const TSharedRef<const FContentBrowserAssetFileItemDataPayload>&)> InAssetPayloadCallback);
 
+	CONTENTBROWSERASSETDATASOURCE_API void EnumerateAssetItemPayloads(const UContentBrowserDataSource* InOwnerDataSource, TArrayView<const FContentBrowserItemData> InItems, TFunctionRef<bool(const TSharedRef<const FContentBrowserAssetFolderItemDataPayload>&)> InFolderPayloadCallback, TFunctionRef<bool(const TSharedRef<const FContentBrowserAssetFileItemDataPayload>&)> InAssetPayloadCallback, TFunctionRef<bool(const TSharedRef<const FContentBrowserUnsupportedAssetFileItemDataPayload>&)> InUnsupportedAssetPayloadCallback);
+
 	CONTENTBROWSERASSETDATASOURCE_API bool IsPrimaryAsset(const FAssetData& InAssetData);
 
 	CONTENTBROWSERASSETDATASOURCE_API bool IsPrimaryAsset(UObject* InObject);
@@ -97,12 +104,16 @@ namespace ContentBrowserAssetData
 
 	CONTENTBROWSERASSETDATASOURCE_API bool CanDeleteAssetFileItem(IAssetTools* InAssetTools, const FContentBrowserAssetFileItemDataPayload& InAssetPayload, FText* OutErrorMsg);
 
+	CONTENTBROWSERASSETDATASOURCE_API bool CanDeleteUnsupportedAssetFileItem(IAssetTools* InAssetTools, const FContentBrowserUnsupportedAssetFileItemDataPayload& InUnsupportedAssetPayload, FText* OutErrorMsg);
+
 	CONTENTBROWSERASSETDATASOURCE_API bool DeleteItems(IAssetTools* InAssetTools, IAssetRegistry* InAssetRegistry, const UContentBrowserDataSource* InOwnerDataSource, TArrayView<const FContentBrowserItemData> InItems);
 
 	CONTENTBROWSERASSETDATASOURCE_API bool DeleteAssetFolderItems(TArrayView<const TSharedRef<const FContentBrowserAssetFolderItemDataPayload>> InFolderPayloads);
 
 	CONTENTBROWSERASSETDATASOURCE_API bool DeleteAssetFileItems(TArrayView<const TSharedRef<const FContentBrowserAssetFileItemDataPayload>> InAssetPayloads);
 
+	CONTENTBROWSERASSETDATASOURCE_API bool DeleteUnsupportedAssetFileItems(TArrayView<const TSharedRef<const FContentBrowserUnsupportedAssetFileItemDataPayload>> InUnsupportedAssetPayloads);
+
 	CONTENTBROWSERASSETDATASOURCE_API bool CanPrivatizeItem(IAssetTools* InAssetTools, IAssetRegistry* InAssetRegistry, const UContentBrowserDataSource* InOwnerDataSource, const FContentBrowserItemData& InItem, FText* OutErrorMsg);
 
 	CONTENTBROWSERASSETDATASOURCE_API bool CanPrivatizeAssetFileItem(IAssetTools* InAssetTools, const FContentBrowserAssetFileItemDataPayload& InAssetPayload, FText* OutErrorMsg);
@@ -147,22 +158,34 @@ namespace ContentBrowserAssetData
 
 	CONTENTBROWSERASSETDATASOURCE_API bool AppendAssetFileItemReference(const FContentBrowserAssetFileItemDataPayload& InAssetPayload, FString& InOutStr);
 
+	CONTENTBROWSERASSETDATASOURCE_API bool AppendUnsupportedAssetFileItemReference(const FContentBrowserUnsupportedAssetFileItemDataPayload& InUnsupportedAssetPayload, FString& InOutStr);
+
 	CONTENTBROWSERASSETDATASOURCE_API bool GetItemPhysicalPath(const UContentBrowserDataSource* InOwnerDataSource, const FContentBrowserItemData& InItem, FString& OutDiskPath);
 
 	CONTENTBROWSERASSETDATASOURCE_API bool GetAssetFolderItemPhysicalPath(const FContentBrowserAssetFolderItemDataPayload& InFolderPayload, FString& OutDiskPath);
 
 	CONTENTBROWSERASSETDATASOURCE_API bool GetAssetFileItemPhysicalPath(const FContentBrowserAssetFileItemDataPayload& InAssetPayload, FString& OutDiskPath);
 
+	CONTENTBROWSERASSETDATASOURCE_API bool GetUnsupportedAssetFileItemPhysicalPath(const FContentBrowserUnsupportedAssetFileItemDataPayload& InUnsupportedAssetPayload, FString& OutDiskPath);
+
 	CONTENTBROWSERASSETDATASOURCE_API bool GetItemAttribute(const UContentBrowserDataSource* InOwnerDataSource, const FContentBrowserItemData& InItem, const bool InIncludeMetaData, const FName InAttributeKey, FContentBrowserItemDataAttributeValue& OutAttributeValue);
 
 	CONTENTBROWSERASSETDATASOURCE_API bool GetAssetFolderItemAttribute(const FContentBrowserAssetFolderItemDataPayload& InFolderPayload, const bool InIncludeMetaData, const FName InAttributeKey, FContentBrowserItemDataAttributeValue& OutAttributeValue);
 
 	CONTENTBROWSERASSETDATASOURCE_API bool GetAssetFileItemAttribute(const FContentBrowserAssetFileItemDataPayload& InAssetPayload, const bool InIncludeMetaData, const FName InAttributeKey, FContentBrowserItemDataAttributeValue& OutAttributeValue);
 
+	CONTENTBROWSERASSETDATASOURCE_API bool GetUnsupportedAssetFileItemAttribute(const FContentBrowserUnsupportedAssetFileItemDataPayload& InUnsupportedAssetPayload, const bool InIncludeMetaData, const FName InAttributeKey, FContentBrowserItemDataAttributeValue& OutAttributeValue);                       
+
+	CONTENTBROWSERASSETDATASOURCE_API bool GetAssetDataAttribute(const FAssetData& InAssetData, const bool InIncludeMetaData, const FName InAttributeKey, FContentBrowserItemDataAttributeValue& OutAttributeValue);
+
 	CONTENTBROWSERASSETDATASOURCE_API bool GetItemAttributes(const UContentBrowserDataSource* InOwnerDataSource, const FContentBrowserItemData& InItem, const bool InIncludeMetaData, FContentBrowserItemDataAttributeValues& OutAttributeValues);
 
 	CONTENTBROWSERASSETDATASOURCE_API bool GetAssetFileItemAttributes(const FContentBrowserAssetFileItemDataPayload& InAssetPayload, const bool InIncludeMetaData, FContentBrowserItemDataAttributeValues& OutAttributeValues);
 
+	CONTENTBROWSERASSETDATASOURCE_API bool GetUnsupportedAssetFileItemAttributes(const FContentBrowserUnsupportedAssetFileItemDataPayload& InUnsupportedAssetPayload, const bool InIncludeMetaData, FContentBrowserItemDataAttributeValues& OutAttributeValues);
+
+	CONTENTBROWSERASSETDATASOURCE_API bool GetAssetDataAttributes(const FAssetData& InAssetData, const bool InIncludeMetaData, FContentBrowserItemDataAttributeValues& OutAttributeValues);
+
 	CONTENTBROWSERASSETDATASOURCE_API void PopulateAssetFolderContextMenu(UContentBrowserDataSource* InOwnerDataSource, UToolMenu* InMenu, FAssetFolderContextMenu& InAssetFolderContextMenu);
 
 	CONTENTBROWSERASSETDATASOURCE_API void PopulateAssetFileContextMenu(UContentBrowserDataSource* InOwnerDataSource, UToolMenu* InMenu, FAssetFileContextMenu& InAssetFileContextMenu);
diff --git a/Engine/Plugins/Editor/ContentBrowser/ContentBrowserAssetDataSource/Source/ContentBrowserAssetDataSource/Public/ContentBrowserAssetDataPayload.h b/Engine/Plugins/Editor/ContentBrowser/ContentBrowserAssetDataSource/Source/ContentBrowserAssetDataSource/Public/ContentBrowserAssetDataPayload.h
--- a/Engine/Plugins/Editor/ContentBrowser/ContentBrowserAssetDataSource/Source/ContentBrowserAssetDataSource/Public/ContentBrowserAssetDataPayload.h
+++ b/Engine/Plugins/Editor/ContentBrowser/ContentBrowserAssetDataSource/Source/ContentBrowserAssetDataSource/Public/ContentBrowserAssetDataPayload.h
@@ -132,6 +132,32 @@ private:
 	TWeakObjectPtr<UObject> SourceObject;
 };
 
+class CONTENTBROWSERASSETDATASOURCE_API FContentBrowserUnsupportedAssetFileItemDataPayload : public IContentBrowserItemDataPayload
+{
+public:
+	// Unsupported asset file but it does have an asset data
+	explicit FContentBrowserUnsupportedAssetFileItemDataPayload(FAssetData&& InAssetData);
+	explicit FContentBrowserUnsupportedAssetFileItemDataPayload(const FAssetData& InAssetData);
+
+	const FAssetData* GetAssetDataIfAvailable() const;
+
+	const FString& GetFilename() const;
+
+	UPackage* GetPackage() const;
+
+private:
+	void FlushCaches() const;
+
+
+	TUniquePtr<FAssetData> OptionalAssetData;
+
+	mutable bool bHasCachedPackagePtr = false;
+	mutable TWeakObjectPtr<UPackage> CachedPackagePtr;
+
+	mutable bool bHasCachedFilename = false;
+	mutable FString CachedFilename;
+};
+
 #if UE_ENABLE_INCLUDE_ORDER_DEPRECATED_IN_5_2
 #include "CoreMinimal.h"
 #endif
diff --git a/Engine/Plugins/Editor/ContentBrowser/ContentBrowserAssetDataSource/Source/ContentBrowserAssetDataSource/Public/ContentBrowserAssetDataSource.h b/Engine/Plugins/Editor/ContentBrowser/ContentBrowserAssetDataSource/Source/ContentBrowserAssetDataSource/Public/ContentBrowserAssetDataSource.h
--- a/Engine/Plugins/Editor/ContentBrowser/ContentBrowserAssetDataSource/Source/ContentBrowserAssetDataSource/Public/ContentBrowserAssetDataSource.h
+++ b/Engine/Plugins/Editor/ContentBrowser/ContentBrowserAssetDataSource/Source/ContentBrowserAssetDataSource/Public/ContentBrowserAssetDataSource.h
@@ -13,6 +13,7 @@
 
 class FContentBrowserAssetFileItemDataPayload;
 class FContentBrowserAssetFolderItemDataPayload;
+class FContentBrowserUnsupportedAssetFileItemDataPayload;
 class FReply;
 struct FPropertyChangedEvent;
 
@@ -55,6 +56,25 @@ public:
 	TArray<FAssetData> CustomSourceAssets;
 };
 
+USTRUCT()
+struct CONTENTBROWSERASSETDATASOURCE_API FContentBrowserCompiledUnsupportedAssetDataFilter
+{
+	GENERATED_BODY()
+
+public:
+	// Filter that are used to determine if an asset is supported
+	FARCompiledFilter ConvertIfFailInclusiveFilter;
+	FARCompiledFilter ConvertIfFailExclusiveFilter;
+
+	// Filter for where to show or not the unsupported assets
+	FARCompiledFilter ShowInclusiveFilter;
+	FARCompiledFilter ShowExclusiveFilter;
+
+	// Standard asset filter modified for the need of the unsupported assets
+	FARCompiledFilter InclusiveFilter;
+	FARCompiledFilter ExclusiveFilter;
+};
+
 UCLASS()
 class CONTENTBROWSERASSETDATASOURCE_API UContentBrowserAssetDataSource : public UContentBrowserDataSource
 {
@@ -81,11 +101,16 @@ public:
 		const FContentBrowserDataClassFilter* ClassFilter = nullptr;
 		const FContentBrowserDataCollectionFilter* CollectionFilter = nullptr;
 
+		// Filter that convert showed items as unsupported items
+		const FContentBrowserDataUnsupportedClassFilter* UnsupportedClassFilter = nullptr;
+
 		const FPathPermissionList* PathPermissionList = nullptr;
 		const FPathPermissionList* ClassPermissionList = nullptr;
 
 		FContentBrowserDataFilterList* FilterList = nullptr;
 		FContentBrowserCompiledAssetDataFilter* AssetDataFilter = nullptr;
+		FContentBrowserCompiledUnsupportedAssetDataFilter* ConvertToUnsupportedAssetDataFilter = nullptr;
+		
 
 		bool bIncludeFolders = false;
 		bool bIncludeFiles = false;
@@ -178,6 +203,8 @@ public:
 
 	virtual bool DoesItemPassFilter(const FContentBrowserItemData& InItem, const FContentBrowserDataCompiledFilter& InFilter) override;
 
+	virtual bool ConvertItemForFilter(FContentBrowserItemData& Item, const FContentBrowserDataCompiledFilter& InFilter) override;
+
 	virtual bool GetItemAttribute(const FContentBrowserItemData& InItem, const bool InIncludeMetaData, const FName InAttributeKey, FContentBrowserItemDataAttributeValue& OutAttributeValue) override;
 
 	virtual bool GetItemAttributes(const FContentBrowserItemData& InItem, const bool InIncludeMetaData, FContentBrowserItemDataAttributeValues& OutAttributeValues) override;
@@ -282,10 +309,14 @@ private:
 
 	FContentBrowserItemData CreateAssetFileItem(const FAssetData& InAssetData);
 
+	FContentBrowserItemData CreateUnsupportedAssetFileItem(const FAssetData& InAssetData);
+
 	TSharedPtr<const FContentBrowserAssetFolderItemDataPayload> GetAssetFolderItemPayload(const FContentBrowserItemData& InItem) const;
 
 	TSharedPtr<const FContentBrowserAssetFileItemDataPayload> GetAssetFileItemPayload(const FContentBrowserItemData& InItem) const;
 
+	TSharedPtr<const FContentBrowserUnsupportedAssetFileItemDataPayload> GetUnsupportedAssetFileItemPayload(const FContentBrowserItemData& InItem) const;
+
 	bool CanHandleDragDropEvent(const FContentBrowserItemData& InItem, const FDragDropEvent& InDragDropEvent) const;
 
 	void OnAssetRegistryFileLoadProgress(const IAssetRegistry::FFileLoadProgressUpdateData& InProgressUpdateData);
diff --git a/Engine/Source/Editor/ContentBrowser/Private/AssetContextMenu.cpp b/Engine/Source/Editor/ContentBrowser/Private/AssetContextMenu.cpp
--- a/Engine/Source/Editor/ContentBrowser/Private/AssetContextMenu.cpp
+++ b/Engine/Source/Editor/ContentBrowser/Private/AssetContextMenu.cpp
@@ -154,47 +154,48 @@ TSharedRef<SWidget> FAssetContextMenu::MakeContextMenu(TArrayView<const FContent
 		// need to keep it here for now to build the correct menu name and register the correct extenders
 
 		// Objects must be loaded for this operation... for now
-		TArray<FAssetData> SelectedAssets;
 		UContentBrowserDataSource* CommonDataSource = nullptr;
-		bool bKeepCheckingCommonDataSource = true;
-		for (const FContentBrowserItem& SelectedItem : SelectedItems)
 		{
-			if (bKeepCheckingCommonDataSource)
+			TArray<FAssetData> SelectedAssets;
+			bool bKeepCheckingCommonDataSource = true;
+			for (const FContentBrowserItem& SelectedItem : SelectedItems)
 			{
-				if (const FContentBrowserItemData* PrimaryInternalItem = SelectedItem.GetPrimaryInternalItem())
+				if (bKeepCheckingCommonDataSource)
 				{
-					if (UContentBrowserDataSource* OwnerDataSource = PrimaryInternalItem->GetOwnerDataSource())
+					if (const FContentBrowserItemData* PrimaryInternalItem = SelectedItem.GetPrimaryInternalItem())
 					{
-						if (CommonDataSource == nullptr)
-						{
-							CommonDataSource = OwnerDataSource;
-						}
-						else if (CommonDataSource != OwnerDataSource)
+						if (UContentBrowserDataSource* OwnerDataSource = PrimaryInternalItem->GetOwnerDataSource())
 						{
-							CommonDataSource = nullptr;
-							bKeepCheckingCommonDataSource = false;
+							if (CommonDataSource == nullptr)
+							{
+								CommonDataSource = OwnerDataSource;
+							}
+							else if (CommonDataSource != OwnerDataSource)
+							{
+								CommonDataSource = nullptr;
+								bKeepCheckingCommonDataSource = false;
+							}
 						}
 					}
 				}
-			}
 
-			FAssetData ItemAssetData;
-			if (SelectedItem.Legacy_TryGetAssetData(ItemAssetData))
-			{
-				SelectedAssets.Add(MoveTemp(ItemAssetData));
+				FAssetData ItemAssetData;
+				if (SelectedItem.Legacy_TryGetAssetData(ItemAssetData))
+				{
+					SelectedAssets.Add(MoveTemp(ItemAssetData));
+				}
 			}
-		}
 
+			ContextObject->bCanBeModified = SelectedAssets.Num() == 0;
+			ContextObject->SelectedAssets = MoveTemp(SelectedAssets);
+		}
 		FAssetToolsModule& AssetToolsModule = FModuleManager::LoadModuleChecked<FAssetToolsModule>("AssetTools");
-		const TSharedRef<FPathPermissionList>& WritableFolderFilter = AssetToolsModule.Get().GetWritableFolderPermissionList();
-
-		ContextObject->bCanBeModified = SelectedAssets.Num() == 0;
+		const TSharedRef<FPathPermissionList>& WritableFolderPermission = AssetToolsModule.Get().GetWritableFolderPermissionList();
 
-		ContextObject->SelectedAssets.Reset();
+		const TArray<FAssetData>& SelectedAssets = ContextObject->SelectedAssets;
 
-		if (SelectedAssets.Num() > 0)
+		if (SelectedAssets.Num() > 0 && SelectedAssets.Num() == SelectedItems.Num())
 		{
-			ContextObject->SelectedAssets.Append(SelectedAssets);
 
 			// Find common class for selected objects
 			UClass* CommonClass = nullptr;
@@ -218,7 +219,7 @@ TSharedRef<SWidget> FAssetContextMenu::MakeContextMenu(TArrayView<const FContent
 			ContextObject->bCanBeModified = true;
 			for (const FAssetData& SelectedAsset : SelectedAssets)
 			{
-				if (WritableFolderFilter->HasFiltering() && !WritableFolderFilter->PassesStartsWithFilter(SelectedAsset.PackageName))
+				if (WritableFolderPermission->HasFiltering() && !WritableFolderPermission->PassesStartsWithFilter(SelectedAsset.PackageName))
 				{
 					ContextObject->bCanBeModified = false;
 					break;
@@ -253,17 +254,35 @@ TSharedRef<SWidget> FAssetContextMenu::MakeContextMenu(TArrayView<const FContent
 				PRAGMA_ENABLE_DEPRECATION_WARNINGS
 			}
 		}
-		else if (SelectedAssets.Num() == 0)
+		else
 		{
 			if (CommonDataSource)
 			{
 				ContextObject->bCanBeModified = true;
 
-				if (WritableFolderFilter->HasFiltering())
+				if (WritableFolderPermission->HasFiltering())
 				{
 					for (const FContentBrowserItem& SelectedItem : SelectedItems)
 					{
-						if (!WritableFolderFilter->PassesStartsWithFilter(SelectedItem.GetVirtualPath()))
+						if (!WritableFolderPermission->PassesStartsWithFilter(SelectedItem.GetInternalPath()))
+						{
+							ContextObject->bCanBeModified = false;
+							break;
+						}
+					}
+				}
+
+				for (const FAssetData& SelectedAsset : SelectedAssets)
+				{
+					if (SelectedAsset.HasAnyPackageFlags(PKG_Cooked | PKG_FilterEditorOnly))
+					{
+						ContextObject->bCanBeModified = false;
+						break;
+					}
+
+					if (const UClass* AssetClass = SelectedAsset.GetClass())
+					{
+						if (AssetClass->IsChildOf<UClass>())
 						{
 							ContextObject->bCanBeModified = false;
 							break;
diff --git a/Engine/Source/Editor/ContentBrowser/Private/SAssetView.cpp b/Engine/Source/Editor/ContentBrowser/Private/SAssetView.cpp
--- a/Engine/Source/Editor/ContentBrowser/Private/SAssetView.cpp
+++ b/Engine/Source/Editor/ContentBrowser/Private/SAssetView.cpp
@@ -267,6 +267,7 @@ void SAssetView::Construct( const FArguments& InArgs )
 	bForceShowEngineContent = InArgs._ForceShowEngineContent;
 	bForceShowPluginContent = InArgs._ForceShowPluginContent;
 	bForceHideScrollbar = InArgs._ForceHideScrollbar;
+	bShowDisallowedAssetClassAsUnsupportedItems = InArgs._ShowDisallowedAssetClassAsUnsupportedItems;
 
 	bPendingUpdateThumbnails = false;
 	bShouldNotifyNextAssetSync = true;
@@ -1876,6 +1877,15 @@ FContentBrowserDataFilter SAssetView::CreateBackendDataFilter() const
 		| (IsShowingLocalizedContent() ? EContentBrowserItemAttributeFilter::IncludeLocalized : EContentBrowserItemAttributeFilter::IncludeNone);
 
 	TSharedPtr<FPathPermissionList> CombinedFolderPermissionList = ContentBrowserUtils::GetCombinedFolderPermissionList(FolderPermissionList, IsShowingReadOnlyFolders() ? nullptr : WritableFolderPermissionList);
+
+	if (bShowDisallowedAssetClassAsUnsupportedItems && AssetClassPermissionList && AssetClassPermissionList->HasFiltering())
+	{
+		// The unsupported item will created as an unsupported asset item instead of normal asset item for the writable folders
+		FContentBrowserDataUnsupportedClassFilter& UnsupportedClassFilter = DataFilter.ExtraFilters.FindOrAddFilter<FContentBrowserDataUnsupportedClassFilter>();
+		UnsupportedClassFilter.ClassPermissionList = AssetClassPermissionList;
+		UnsupportedClassFilter.FolderPermissionList = WritableFolderPermissionList;
+	}
+
 	ContentBrowserUtils::AppendAssetFilterToContentBrowserFilter(BackendFilter, AssetClassPermissionList, CombinedFolderPermissionList, DataFilter);
 
 	if (bHasCollections && !SourcesData.IsDynamicCollection())
@@ -4772,13 +4782,13 @@ void SAssetView::HandleItemDataUpdated(TArrayView<const FContentBrowserItemDataU
 	bool bRefreshView = false;
 	TSet<TSharedPtr<FAssetViewItem>> ItemsPendingInplaceFrontendFilter;
 
-	auto AddItem = [this, &ItemsPendingInplaceFrontendFilter](const FContentBrowserItemKey& InItemDataKey, const FContentBrowserItemData& InItemData)
+	auto AddItem = [this, &ItemsPendingInplaceFrontendFilter](const FContentBrowserItemKey& InItemDataKey, FContentBrowserItemData&& InItemData)
 	{
 		TSharedPtr<FAssetViewItem>& ItemToUpdate = AvailableBackendItems.FindOrAdd(InItemDataKey);
 		if (ItemToUpdate)
 		{
 			// Update the item
-			ItemToUpdate->AppendItemData(InItemData);
+			ItemToUpdate->AppendItemData(MoveTemp(InItemData));
 
 			// Update the custom column data
 			ItemToUpdate->CacheCustomColumns(CustomColumns, true, true, true);
@@ -4789,7 +4799,7 @@ void SAssetView::HandleItemDataUpdated(TArrayView<const FContentBrowserItemDataU
 		}
 		else
 		{
-			ItemToUpdate = MakeShared<FAssetViewItem>(InItemData);
+			ItemToUpdate = MakeShared<FAssetViewItem>(MoveTemp(InItemData));
 
 			// This item is new so put it in the pending set to be processed over time
 			ItemsPendingFrontendFilter.Add(ItemToUpdate);
@@ -4827,38 +4837,51 @@ void SAssetView::HandleItemDataUpdated(TArrayView<const FContentBrowserItemDataU
 		}
 	};
 
-	auto DoesItemPassBackendFilter = [this, &CompiledDataFilters](const FContentBrowserItemData& InItemData)
+	auto GetBackendFilterCompliantItem = [this, &CompiledDataFilters](const FContentBrowserItemData& InItemData, bool& bOutPassFilter)
 	{
 		UContentBrowserDataSource* ItemDataSource = InItemData.GetOwnerDataSource();
+		FContentBrowserItemData ItemData = InItemData;
 		for (const FContentBrowserDataCompiledFilter& DataFilter : CompiledDataFilters)
 		{
-			if (ItemDataSource->DoesItemPassFilter(InItemData, DataFilter))
+			// We only convert the item if this is the right filter for the data source
+			if (ItemDataSource->ConvertItemForFilter(ItemData, DataFilter))
 			{
-				return true;
+				bOutPassFilter = ItemDataSource->DoesItemPassFilter(ItemData, DataFilter);
+
+				return ItemData;
+			}
+
+			if (ItemDataSource->DoesItemPassFilter(ItemData, DataFilter))
+			{
+				bOutPassFilter = true;
+				return ItemData;
 			}
 		}
-		return false;
+
+		bOutPassFilter = false;
+		return ItemData;
 	};
 
 	// Process the main set of updates
 	for (const FContentBrowserItemDataUpdate& ItemDataUpdate : InUpdatedItems)
 	{
-		const FContentBrowserItemData& ItemData = ItemDataUpdate.GetItemData();
+		bool bItemPassFilter = false;
+		FContentBrowserItemData ItemData = GetBackendFilterCompliantItem(ItemDataUpdate.GetItemData(), bItemPassFilter);
 		const FContentBrowserItemKey ItemDataKey(ItemData);
 
 		switch (ItemDataUpdate.GetUpdateType())
 		{
 		case EContentBrowserItemUpdateType::Added:
-			if (DoesItemPassBackendFilter(ItemData))
+			if (bItemPassFilter)
 			{
-				AddItem(ItemDataKey, ItemData);
+				AddItem(ItemDataKey, MoveTemp(ItemData));
 			}
 			break;
 
 		case EContentBrowserItemUpdateType::Modified:
-			if (DoesItemPassBackendFilter(ItemData))
+			if (bItemPassFilter)
 			{
-				AddItem(ItemDataKey, ItemData);
+				AddItem(ItemDataKey, MoveTemp(ItemData));
 			}
 			else
 			{
@@ -4872,9 +4895,9 @@ void SAssetView::HandleItemDataUpdated(TArrayView<const FContentBrowserItemDataU
 				const FContentBrowserItemKey OldItemDataKey(OldMinimalItemData);
 				RemoveItem(OldItemDataKey, OldMinimalItemData);
 
-				if (DoesItemPassBackendFilter(ItemData))
+				if (bItemPassFilter)
 				{
-					AddItem(ItemDataKey, ItemData);
+					AddItem(ItemDataKey, MoveTemp(ItemData));
 				}
 				else
 				{
diff --git a/Engine/Source/Editor/ContentBrowser/Private/SContentBrowser.cpp b/Engine/Source/Editor/ContentBrowser/Private/SContentBrowser.cpp
--- a/Engine/Source/Editor/ContentBrowser/Private/SContentBrowser.cpp
+++ b/Engine/Source/Editor/ContentBrowser/Private/SContentBrowser.cpp
@@ -274,7 +274,8 @@ void SContentBrowser::Construct( const FArguments& InArgs, const FName& InInstan
 		.AddMetaData<FTagMetaData>(FTagMetaData(TEXT("ContentBrowserAssets")))
 		.OwningContentBrowser(SharedThis(this))
 		.OnSearchOptionsChanged(this, &SContentBrowser::HandleAssetViewSearchOptionsChanged)
-		.FillEmptySpaceInTileView(true);
+		.FillEmptySpaceInTileView(true)
+		.ShowDisallowedAssetClassAsUnsupportedItems(true);
 
 
 	TSharedRef<SWidget> ViewOptions = SNullWidget::NullWidget;
diff --git a/Engine/Source/Editor/ContentBrowser/Private/SFilterList.cpp b/Engine/Source/Editor/ContentBrowser/Private/SFilterList.cpp
--- a/Engine/Source/Editor/ContentBrowser/Private/SFilterList.cpp
+++ b/Engine/Source/Editor/ContentBrowser/Private/SFilterList.cpp
@@ -229,9 +229,13 @@ void SFilterList::DisableFiltersThatHideItems(TArrayView<const FContentBrowserIt
 				}
 
 				FContentBrowserItem::FItemDataArrayView InternalItems = Item.GetInternalItems();
-				for (const FContentBrowserItemData& InternalItem : InternalItems)
+				for (const FContentBrowserItemData& InternalItemRef : InternalItems)
 				{
-					UContentBrowserDataSource* ItemDataSource = InternalItem.GetOwnerDataSource();
+					UContentBrowserDataSource* ItemDataSource = InternalItemRef.GetOwnerDataSource();
+
+					FContentBrowserItemData InternalItem = InternalItemRef;
+					ItemDataSource->ConvertItemForFilter(InternalItem, CompiledDataFilter);
+
 					if (!ItemDataSource->DoesItemPassFilter(InternalItem, CompiledDataFilter))
 					{
 						bDisableAllBackendFilters = true;
diff --git a/Engine/Source/Editor/ContentBrowser/Private/SPathView.cpp b/Engine/Source/Editor/ContentBrowser/Private/SPathView.cpp
--- a/Engine/Source/Editor/ContentBrowser/Private/SPathView.cpp
+++ b/Engine/Source/Editor/ContentBrowser/Private/SPathView.cpp
@@ -1955,27 +1955,30 @@ void SPathView::HandleItemDataUpdated(TArrayView<const FContentBrowserItemDataUp
 
 	for (const FContentBrowserItemDataUpdate& ItemDataUpdate : InUpdatedItems)
 	{
-		const FContentBrowserItemData& ItemData = ItemDataUpdate.GetItemData();
-		if (!ItemData.IsFolder())
+		const FContentBrowserItemData& ItemDataRef = ItemDataUpdate.GetItemData();
+		if (!ItemDataRef.IsFolder())
 		{
 			continue;
 		}
 
 		ConditionalCompileFilter();
 
+		FContentBrowserItemData ItemData = ItemDataRef;
+		ItemData.GetOwnerDataSource()->ConvertItemForFilter(ItemData, CompiledDataFilter);
+
 		switch (ItemDataUpdate.GetUpdateType())
 		{
 		case EContentBrowserItemUpdateType::Added:
 			if (DoesItemPassFilter(ItemData))
 			{
-				AddFolderItem(CopyTemp(ItemData));
+				AddFolderItem(MoveTemp(ItemData));
 			}
 			break;
 
 		case EContentBrowserItemUpdateType::Modified:
 			if (DoesItemPassFilter(ItemData))
 			{
-				AddFolderItem(CopyTemp(ItemData));
+				AddFolderItem(MoveTemp(ItemData));
 			}
 			else
 			{
@@ -1990,7 +1993,7 @@ void SPathView::HandleItemDataUpdated(TArrayView<const FContentBrowserItemDataUp
 
 			if (DoesItemPassFilter(ItemData))
 			{
-				AddFolderItem(CopyTemp(ItemData));
+				AddFolderItem(MoveTemp(ItemData));
 			}
 		}
 		break;
@@ -2246,6 +2249,7 @@ void SFavoritePathView::Populate(const bool bIsRefreshingFilter)
 			[this, &CompiledDataFilter](FContentBrowserItemData&& InItemData)
 				{
 					UContentBrowserDataSource* ItemDataSource = InItemData.GetOwnerDataSource();
+					ItemDataSource->ConvertItemForFilter(InItemData, CompiledDataFilter);
 					if (ItemDataSource->DoesItemPassFilter(InItemData, CompiledDataFilter))
 					{
 						if (TSharedPtr<FTreeItem> Item = AddFolderItem(MoveTemp(InItemData)))
@@ -2450,27 +2454,30 @@ void SFavoritePathView::HandleItemDataUpdated(TArrayView<const FContentBrowserIt
 
 	for (const FContentBrowserItemDataUpdate& ItemDataUpdate : InUpdatedItems)
 	{
-		const FContentBrowserItemData& ItemData = ItemDataUpdate.GetItemData();
-		if (!ItemData.IsFolder())
+		const FContentBrowserItemData& ItemDataRef = ItemDataUpdate.GetItemData();
+		if (!ItemDataRef.IsFolder())
 		{
 			continue;
 		}
 
 		ConditionalCompileFilter();
 
+		FContentBrowserItemData ItemData = ItemDataUpdate.GetItemData();
+		ItemData.GetOwnerDataSource()->ConvertItemForFilter(ItemData, CompiledDataFilter);
+
 		switch (ItemDataUpdate.GetUpdateType())
 		{
 		case EContentBrowserItemUpdateType::Added:
 			if (DoesItemPassFilter(ItemData))
 			{
-				AddFolderItem(CopyTemp(ItemData));
+				AddFolderItem(MoveTemp(ItemData));
 			}
 			break;
 
 		case EContentBrowserItemUpdateType::Modified:
 			if (DoesItemPassFilter(ItemData))
 			{
-				AddFolderItem(CopyTemp(ItemData));
+				AddFolderItem(MoveTemp(ItemData));
 			}
 			else
 			{
@@ -2485,7 +2492,7 @@ void SFavoritePathView::HandleItemDataUpdated(TArrayView<const FContentBrowserIt
 
 			if (DoesItemPassFilter(ItemData))
 			{
-				AddFolderItem(CopyTemp(ItemData));
+				AddFolderItem(MoveTemp(ItemData));
 			}
 
 			ContentBrowserUtils::RemoveFavoriteFolder(ItemDataUpdate.GetPreviousVirtualPath().ToString());
diff --git a/Engine/Source/Editor/ContentBrowser/Public/SAssetView.h b/Engine/Source/Editor/ContentBrowser/Public/SAssetView.h
--- a/Engine/Source/Editor/ContentBrowser/Public/SAssetView.h
+++ b/Engine/Source/Editor/ContentBrowser/Public/SAssetView.h
@@ -121,6 +121,7 @@ public:
 		, _ForceShowEngineContent(false)
 		, _ForceShowPluginContent(false)
 		, _ForceHideScrollbar(false)
+		, _ShowDisallowedAssetClassAsUnsupportedItems(false)
 		{}
 
 		/** Called to check if an asset should be filtered out by external code */
@@ -264,6 +265,9 @@ public:
 		/** Should always hide scrollbar (Removes scrollbar) */
 		SLATE_ARGUMENT(bool, ForceHideScrollbar)
 
+		/** Allow the asset view to display the hidden asset class as unsupported items */
+		SLATE_ARGUMENT(bool, ShowDisallowedAssetClassAsUnsupportedItems)
+
 		/** Called to check if an asset tag should be display in details view. */
 		SLATE_EVENT( FOnShouldDisplayAssetTag, OnAssetTagWantsToBeDisplayed )
 
@@ -1141,6 +1145,8 @@ private:
 	/** Initial set of item categories that this view should show - may be adjusted further by things like CanShowClasses or legacy delegate bindings */
 	EContentBrowserItemCategoryFilter InitialCategoryFilter;
 
+	bool bShowDisallowedAssetClassAsUnsupportedItems = false;
+
 	/** A struct to hold data for the deferred creation of a file or folder item */
 	struct FCreateDeferredItemData
 	{
diff --git a/Engine/Source/Editor/ContentBrowserData/Private/ContentBrowserDataSource.cpp b/Engine/Source/Editor/ContentBrowserData/Private/ContentBrowserDataSource.cpp
--- a/Engine/Source/Editor/ContentBrowserData/Private/ContentBrowserDataSource.cpp
+++ b/Engine/Source/Editor/ContentBrowserData/Private/ContentBrowserDataSource.cpp
@@ -265,6 +265,11 @@ bool UContentBrowserDataSource::DoesItemPassFilter(const FContentBrowserItemData
 	return false;
 }
 
+bool UContentBrowserDataSource::ConvertItemForFilter(FContentBrowserItemData& Item, const FContentBrowserDataCompiledFilter& InFilter)
+{
+	return false;
+}
+
 bool UContentBrowserDataSource::GetItemAttribute(const FContentBrowserItemData& InItem, const bool InIncludeMetaData, const FName InAttributeKey, FContentBrowserItemDataAttributeValue& OutAttributeValue)
 {
 	return false;
diff --git a/Engine/Source/Editor/ContentBrowserData/Private/ContentBrowserItem.cpp b/Engine/Source/Editor/ContentBrowserData/Private/ContentBrowserItem.cpp
--- a/Engine/Source/Editor/ContentBrowserData/Private/ContentBrowserItem.cpp
+++ b/Engine/Source/Editor/ContentBrowserData/Private/ContentBrowserItem.cpp
@@ -290,6 +290,12 @@ bool FContentBrowserItem::IsFile() const
 	return PrimaryItemData && PrimaryItemData->IsFile();
 }
 
+bool FContentBrowserItem::IsSupported() const
+{
+	const FContentBrowserItemData* PrimaryItemData = GetPrimaryInternalItem();
+	return PrimaryItemData && PrimaryItemData->IsSupported();
+}
+
 bool FContentBrowserItem::IsTemporary() const
 {
 	const FContentBrowserItemData* PrimaryItemData = GetPrimaryInternalItem();
diff --git a/Engine/Source/Editor/ContentBrowserData/Private/ContentBrowserItemData.cpp b/Engine/Source/Editor/ContentBrowserData/Private/ContentBrowserItemData.cpp
--- a/Engine/Source/Editor/ContentBrowserData/Private/ContentBrowserItemData.cpp
+++ b/Engine/Source/Editor/ContentBrowserData/Private/ContentBrowserItemData.cpp
@@ -48,6 +48,11 @@ bool FContentBrowserItemData::IsFile() const
 	return EnumHasAnyFlags(ItemFlags, EContentBrowserItemFlags::Type_File);
 }
 
+bool FContentBrowserItemData::IsSupported() const
+{
+	return !EnumHasAnyFlags(ItemFlags, EContentBrowserItemFlags::Misc_Unsupported);
+}
+
 bool FContentBrowserItemData::IsTemporary() const
 {
 	return EnumHasAnyFlags(ItemFlags, EContentBrowserItemFlags::Temporary_MASK);
diff --git a/Engine/Source/Editor/ContentBrowserData/Public/ContentBrowserDataFilter.h b/Engine/Source/Editor/ContentBrowserData/Public/ContentBrowserDataFilter.h
--- a/Engine/Source/Editor/ContentBrowserData/Public/ContentBrowserDataFilter.h
+++ b/Engine/Source/Editor/ContentBrowserData/Public/ContentBrowserDataFilter.h
@@ -309,3 +309,19 @@ public:
 	UPROPERTY(EditAnywhere, BlueprintReadWrite, Category="ContentBrowser")
 	bool bIncludeChildCollections = false;
 };
+
+
+/**
+ * Data used to tell the content browser to show the item that doesn't pass the class permission list as a unsupported asset
+ * @note This will restrict user interaction with those asset in the content browser and only affect the asset in the folders specified in the permission list
+ */
+USTRUCT()
+struct CONTENTBROWSERDATA_API FContentBrowserDataUnsupportedClassFilter
+{
+	GENERATED_BODY()
+
+public:
+	TSharedPtr<FPathPermissionList> ClassPermissionList;
+
+	TSharedPtr<FPathPermissionList> FolderPermissionList;
+};
diff --git a/Engine/Source/Editor/ContentBrowserData/Public/ContentBrowserDataSource.h b/Engine/Source/Editor/ContentBrowserData/Public/ContentBrowserDataSource.h
--- a/Engine/Source/Editor/ContentBrowserData/Public/ContentBrowserDataSource.h
+++ b/Engine/Source/Editor/ContentBrowserData/Public/ContentBrowserDataSource.h
@@ -274,7 +274,8 @@ public:
 	virtual bool CreateFolder(const FName InPath, FContentBrowserItemDataTemporaryContext& OutPendingItem);
 
 	/*
-	 * Query whether the given item passes the given compiled filter.
+	 * Query whether the given item passes the given compiled filter. Should be called after ConvertItemForFilter
+	 *
 	 * @see CompileFilter.
 	 *
 	 * @param InItem The item to query.
@@ -284,6 +285,19 @@ public:
 	 */
 	virtual bool DoesItemPassFilter(const FContentBrowserItemData& InItem, const FContentBrowserDataCompiledFilter& InFilter);
 
+
+	/*
+	 * Let the compiled filter decide the payload and the type of the item 
+	 * Some Compiled filter might change the type/payload of the item. This allow these filter to work properly and should be called before the filtering (see DoesItemPassFilter)
+	 * @see CompileFilter
+	 * 
+	 * @param Item The item that might be converted
+	 * @param InFilter The compiled filter used to possibly convert the matching items.
+	 * 
+	 * @return True if the item was converted by the filter.
+	 */
+	virtual bool ConvertItemForFilter(FContentBrowserItemData& Item, const FContentBrowserDataCompiledFilter& InFilter);
+
 	/**
 	 * Query the value of the given attribute on the given item.
 	 *
diff --git a/Engine/Source/Editor/ContentBrowserData/Public/ContentBrowserItem.h b/Engine/Source/Editor/ContentBrowserData/Public/ContentBrowserItem.h
--- a/Engine/Source/Editor/ContentBrowserData/Public/ContentBrowserItem.h
+++ b/Engine/Source/Editor/ContentBrowserData/Public/ContentBrowserItem.h
@@ -135,6 +135,13 @@ public:
 	 */
 	bool IsFile() const;
 
+	/**
+	 * Check if the item is representing a supported item
+	 * The content browser can also display some unsupported asset
+	 * @note Equivalent to testing whether EContentBrowserItemFlags::Misc_Unsupported is not set on GetItemFlags()
+	 */
+	bool IsSupported() const;
+
 	/**
 	 * Check to see whether this item is temporary.
 	 * @note Equivalent to testing whether any of EContentBrowserItemFlags::Temporary_MASK is set on GetItemFlags().
diff --git a/Engine/Source/Editor/ContentBrowserData/Public/ContentBrowserItemData.h b/Engine/Source/Editor/ContentBrowserData/Public/ContentBrowserItemData.h
--- a/Engine/Source/Editor/ContentBrowserData/Public/ContentBrowserItemData.h
+++ b/Engine/Source/Editor/ContentBrowserData/Public/ContentBrowserItemData.h
@@ -23,7 +23,7 @@ class IContentBrowserItemDataSink;
 class UContentBrowserDataSource;
 
 /** Flags denoting basic state information for an item instance */
-enum class EContentBrowserItemFlags : uint8
+enum class EContentBrowserItemFlags : uint16
 {
 	/** No flags */
 	None = 0,
@@ -44,6 +44,9 @@ enum class EContentBrowserItemFlags : uint8
 	Temporary_Creation = 1<<6,
 	Temporary_Duplication = 1<<7,
 	Temporary_MASK = Temporary_Creation | Temporary_Duplication,
+
+	/** Flag to mark the item as showing something that is unsupported */
+	Misc_Unsupported = 1 << 8,
 };
 ENUM_CLASS_FLAGS(EContentBrowserItemFlags);
 
@@ -165,6 +168,13 @@ public:
 	 */
 	bool IsFile() const;
 
+	/**
+	 * Check if the item is representing a supported item
+	 * The content browser can also display some unsupported asset
+	 * @note Equivalent to testing whether EContentBrowserItemFlags::Misc_Unsupported is not set on GetItemFlags()
+	 */
+	bool IsSupported() const;
+
 	/**
 	 * Check to see whether this item is temporary.
 	 * @note Equivalent to testing whether any of EContentBrowserItemFlags::Temporary_MASK is set on GetItemFlags().
