Add more scope traces to investigate UContentBrowserDataSubsystem::Tick() hitches during PIE

#swarm 24363495

diff --git a/Engine/Plugins/Editor/ContentBrowser/ContentBrowserAssetDataSource/Source/ContentBrowserAssetDataSource/Private/ContentBrowserAssetDataSource.cpp b/Engine/Plugins/Editor/ContentBrowser/ContentBrowserAssetDataSource/Source/ContentBrowserAssetDataSource/Private/ContentBrowserAssetDataSource.cpp
--- a/Engine/Plugins/Editor/ContentBrowser/ContentBrowserAssetDataSource/Source/ContentBrowserAssetDataSource/Private/ContentBrowserAssetDataSource.cpp
+++ b/Engine/Plugins/Editor/ContentBrowser/ContentBrowserAssetDataSource/Source/ContentBrowserAssetDataSource/Private/ContentBrowserAssetDataSource.cpp
@@ -407,6 +407,8 @@ bool UContentBrowserAssetDataSource::CreatePathFilter(FAssetFilterInputParams& P
 
 bool UContentBrowserAssetDataSource::CreateAssetFilter(FAssetFilterInputParams& Params, FName InPath, const FContentBrowserDataFilter& InFilter, FContentBrowserDataCompiledFilter& OutCompiledFilter, FCompileARFilterFunc CreateCompiledFilter)
 {
+	TRACE_CPUPROFILER_EVENT_SCOPE(UContentBrowserAssetDataSource::CreateAssetFilter);
+
 	// If we're not including files, then we can bail now as the rest of this function deals with assets
 	if (!Params.bIncludeFiles)
 	{
@@ -526,14 +528,18 @@ PRAGMA_ENABLE_DEPRECATION_WARNINGS
 				InternalPathFilter.bRecursivePaths = InFilter.bRecursivePaths;
 				CreateCompiledFilter(InternalPathFilter, CompiledInternalPathFilter);
 
-				// Remove paths that do not pass item attribute filter (Engine, Plugins, Developer, Localized, __ExternalActors__ etc..)
-				for (auto It = CompiledInternalPathFilter.PackagePaths.CreateIterator(); It; ++It)
 				{
-					FNameBuilder PathStr(*It);
-					FStringView Path(PathStr);
-					if (!ContentBrowserDataUtils::PathPassesAttributeFilter(Path, 0, InFilter.ItemAttributeFilter))
+					TRACE_CPUPROFILER_EVENT_SCOPE(UContentBrowserAssetDataSource::CreateAssetFilter::RemovePaths);
+			
+					// Remove paths that do not pass item attribute filter (Engine, Plugins, Developer, Localized, __ExternalActors__ etc..)
+					for (auto It = CompiledInternalPathFilter.PackagePaths.CreateIterator(); It; ++It)
 					{
-						It.RemoveCurrent();
+						FNameBuilder PathStr(*It);
+						FStringView Path(PathStr);
+						if (!ContentBrowserDataUtils::PathPassesAttributeFilter(Path, 0, InFilter.ItemAttributeFilter))
+						{
+							It.RemoveCurrent();
+						}
 					}
 				}
 			}
diff --git a/Engine/Source/Runtime/AssetRegistry/Private/AssetRegistry.cpp b/Engine/Source/Runtime/AssetRegistry/Private/AssetRegistry.cpp
--- a/Engine/Source/Runtime/AssetRegistry/Private/AssetRegistry.cpp
+++ b/Engine/Source/Runtime/AssetRegistry/Private/AssetRegistry.cpp
@@ -3162,6 +3162,8 @@ namespace UE::AssetRegistry
 void FAssetRegistryImpl::CompileFilter(Impl::FClassInheritanceContext& InheritanceContext, const FARFilter& InFilter,
 	FARCompiledFilter& OutCompiledFilter) const
 {
+	TRACE_CPUPROFILER_EVENT_SCOPE(FAssetRegistryImpl::CompileFilter);
+
 	OutCompiledFilter.Clear();
 	OutCompiledFilter.PackageNames.Append(InFilter.PackageNames);
 	OutCompiledFilter.PackagePaths.Reserve(InFilter.PackagePaths.Num());
@@ -3206,6 +3208,8 @@ PRAGMA_ENABLE_DEPRECATION_WARNINGS
 
 	if (InFilter.bRecursivePaths)
 	{
+		TRACE_CPUPROFILER_EVENT_SCOPE(FAssetRegistryImpl::CompileFilter::AddPaths);
+		
 		// Add the sub-paths of all the input paths to the expanded list
 		for (const FName& PackagePath : InFilter.PackagePaths)
 		{
