Checking for external assets in IsPrimaryAsset to avoid content browser updates on external assets

diff --git a/Engine/Plugins/Editor/ContentBrowser/ContentBrowserAssetDataSource/Source/ContentBrowserAssetDataSource/Private/ContentBrowserAssetDataCore.cpp b/Engine/Plugins/Editor/ContentBrowser/ContentBrowserAssetDataSource/Source/ContentBrowserAssetDataSource/Private/ContentBrowserAssetDataCore.cpp
--- a/Engine/Plugins/Editor/ContentBrowser/ContentBrowserAssetDataSource/Source/ContentBrowserAssetDataSource/Private/ContentBrowserAssetDataCore.cpp
+++ b/Engine/Plugins/Editor/ContentBrowser/ContentBrowserAssetDataSource/Source/ContentBrowserAssetDataSource/Private/ContentBrowserAssetDataCore.cpp
@@ -146,14 +146,20 @@ void EnumerateAssetItemPayloads(const UContentBrowserDataSource* InOwnerDataSour
 
 bool IsPrimaryAsset(const FAssetData& InAssetData)
 {
+	// External assets are not displayed in the Content Browser or other asset pickers
+	const bool bIsExternalAsset = !InAssetData.GetOptionalOuterPathName().IsNone();
+
 	// Check for the asset being a redirector first, as currently only class 
 	// redirectors emit non-primary assets from the Asset Registry
-	return !InAssetData.IsRedirector() || InAssetData.IsUAsset();
+	return !InAssetData.IsRedirector() && InAssetData.IsUAsset() && !bIsExternalAsset;
 }
 
 bool IsPrimaryAsset(UObject* InObject)
 {
-	return !FAssetData::IsRedirector(InObject) && FAssetData::IsUAsset(InObject);
+	// External assets are not displayed in the Content Browser or other asset pickers
+	const bool bIsExternalAsset = InObject->IsPackageExternal();
+
+	return !FAssetData::IsRedirector(InObject) && FAssetData::IsUAsset(InObject) && !bIsExternalAsset;
 }
 
 void SetOptionalErrorMessage(FText* OutErrorMsg, FText InErrorMsg)
