The content browser can now respond to packages changing on disk.

- In the last year we have started to make changes to the package format to allow data to be stored in the same file, but not actually inside of the package itself (See FPackageTrailer) and we might consider moving some other data blobs there in the future. This change has meant that there are code paths that can change things with the package file, that don't actually call save package, such as the process that virtualizes assets and can leave incorrect, stale data in the content browser tool tip.
-- This means when a user virtualizes an asset, the tooltip might indicate that it has not been virtualized and the size of the file on disk has not been changed until they browser to a different folder and browse back to refresh. This has caused people to believe that the virtualization operation has failed.
- This change hooks up UContentBrowserAssetDataSource to trigger a EContentBrowserItemUpdateType::Modified event when ever the package file changes on disk.
- During package save this could trigger multiple EContentBrowserItemUpdateType::Modified events for the same package (one from saving and one from the file being modified) but it looks like duplicate events already occur (if the packages are checked out as part of the save, this will introduce another EContentBrowserItemUpdateType::Modified event). There is a TODO in the code that suggests we might want to spend time removing duplicates events before broadcasting the delegate but when profiling I found that the overhead is quite low. In one test, saving 20 packages cost 4ms for the events triggered by the save and an additional 4ms for the events triggered by the file changing on disk. With the margins being so low it still doesn't look like a problem that is worth spending time on, however I did add a TRACE_CPUPROFILER_EVENT_SCOPE to the broadcasting of the delegate (UContentBrowserDataSubsystem::ItemDataUpdatedDelegate) so if it does become a problem in the future it will be easier to track down.

diff --git a/Engine/Plugins/Editor/ContentBrowser/ContentBrowserAssetDataSource/Source/ContentBrowserAssetDataSource/Private/ContentBrowserAssetDataSource.cpp b/Engine/Plugins/Editor/ContentBrowser/ContentBrowserAssetDataSource/Source/ContentBrowserAssetDataSource/Private/ContentBrowserAssetDataSource.cpp
--- a/Engine/Plugins/Editor/ContentBrowser/ContentBrowserAssetDataSource/Source/ContentBrowserAssetDataSource/Private/ContentBrowserAssetDataSource.cpp
+++ b/Engine/Plugins/Editor/ContentBrowser/ContentBrowserAssetDataSource/Source/ContentBrowserAssetDataSource/Private/ContentBrowserAssetDataSource.cpp
@@ -60,6 +60,7 @@ void UContentBrowserAssetDataSource::Initialize(const bool InAutoRegister)
 	AssetRegistry->OnAssetRemoved().AddUObject(this, &UContentBrowserAssetDataSource::OnAssetRemoved);
 	AssetRegistry->OnAssetRenamed().AddUObject(this, &UContentBrowserAssetDataSource::OnAssetRenamed);
 	AssetRegistry->OnAssetUpdated().AddUObject(this, &UContentBrowserAssetDataSource::OnAssetUpdated);
+	AssetRegistry->OnAssetUpdatedOnDisk().AddUObject(this, &UContentBrowserAssetDataSource::OnAssetUpdatedOnDisk);
 	AssetRegistry->OnPathAdded().AddUObject(this, &UContentBrowserAssetDataSource::OnPathAdded);
 	AssetRegistry->OnPathRemoved().AddUObject(this, &UContentBrowserAssetDataSource::OnPathRemoved);
 	AssetRegistry->OnFilesLoaded().AddUObject(this, &UContentBrowserAssetDataSource::OnScanCompleted);
@@ -187,6 +188,7 @@ void UContentBrowserAssetDataSource::Shutdown()
 			AssetRegistryMaybe->OnAssetRemoved().RemoveAll(this);
 			AssetRegistryMaybe->OnAssetRenamed().RemoveAll(this);
 			AssetRegistryMaybe->OnAssetUpdated().RemoveAll(this);
+			AssetRegistryMaybe->OnAssetUpdatedOnDisk().RemoveAll(this);	
 			AssetRegistryMaybe->OnPathAdded().RemoveAll(this);
 			AssetRegistryMaybe->OnPathRemoved().RemoveAll(this);
 			AssetRegistryMaybe->OnFilesLoaded().RemoveAll(this);
@@ -2160,6 +2162,14 @@ void UContentBrowserAssetDataSource::OnAssetUpdated(const FAssetData& InAssetDat
 	}
 }
 
+void UContentBrowserAssetDataSource::OnAssetUpdatedOnDisk(const FAssetData& InAssetData)
+{
+	if (ContentBrowserAssetData::IsPrimaryAsset(InAssetData))
+	{
+		QueueItemDataUpdate(FContentBrowserItemDataUpdate::MakeItemModifiedUpdate(CreateAssetFileItem(InAssetData)));
+	}
+}
+
 void UContentBrowserAssetDataSource::OnObjectPropertyChanged(UObject* InObject, FPropertyChangedEvent& InPropertyChangedEvent)
 {
 	if (InObject && InObject->IsAsset() && ContentBrowserAssetData::IsPrimaryAsset(InObject))
diff --git a/Engine/Plugins/Editor/ContentBrowser/ContentBrowserAssetDataSource/Source/ContentBrowserAssetDataSource/Public/ContentBrowserAssetDataSource.h b/Engine/Plugins/Editor/ContentBrowser/ContentBrowserAssetDataSource/Source/ContentBrowserAssetDataSource/Public/ContentBrowserAssetDataSource.h
--- a/Engine/Plugins/Editor/ContentBrowser/ContentBrowserAssetDataSource/Source/ContentBrowserAssetDataSource/Public/ContentBrowserAssetDataSource.h
+++ b/Engine/Plugins/Editor/ContentBrowser/ContentBrowserAssetDataSource/Source/ContentBrowserAssetDataSource/Public/ContentBrowserAssetDataSource.h
@@ -329,6 +329,8 @@ private:
 
 	void OnAssetUpdated(const FAssetData& InAssetData);
 
+	void OnAssetUpdatedOnDisk(const FAssetData& InAssetData);
+
 	void OnObjectPropertyChanged(UObject* InObject, FPropertyChangedEvent& InPropertyChangedEvent);
 
 	void OnObjectPreSave(UObject* InObject, class FObjectPreSaveContext InObjectPreSaveContext);
diff --git a/Engine/Source/Editor/ContentBrowserData/Private/ContentBrowserDataSubsystem.cpp b/Engine/Source/Editor/ContentBrowserData/Private/ContentBrowserDataSubsystem.cpp
--- a/Engine/Source/Editor/ContentBrowserData/Private/ContentBrowserDataSubsystem.cpp
+++ b/Engine/Source/Editor/ContentBrowserData/Private/ContentBrowserDataSubsystem.cpp
@@ -707,6 +707,8 @@ void UContentBrowserDataSubsystem::Tick(const float InDeltaTime)
 
 	if (PendingUpdates.Num() > 0)
 	{
+		TRACE_CPUPROFILER_EVENT_SCOPE(UContentBrowserDataSubsystem::BroadCastItemDataUpdate);
+
 		TArray<FContentBrowserItemDataUpdate> LocalPendingUpdates = MoveTemp(PendingUpdates);
 		PendingUpdates.Empty();
 		ItemDataUpdatedDelegate.Broadcast(MakeArrayView(LocalPendingUpdates));
