Asset data event generated by the asset redirectors are considered as UAssset and they need don't a special handling from the ContentBrowser core functionnality.

This also fix an issue with the redirectors not being displayed by ContentBrowser.

#rnx

diff --git a/Engine/Plugins/Editor/ContentBrowser/ContentBrowserAssetDataSource/Source/ContentBrowserAssetDataSource/Private/ContentBrowserAssetDataCore.cpp b/Engine/Plugins/Editor/ContentBrowser/ContentBrowserAssetDataSource/Source/ContentBrowserAssetDataSource/Private/ContentBrowserAssetDataCore.cpp
--- a/Engine/Plugins/Editor/ContentBrowser/ContentBrowserAssetDataSource/Source/ContentBrowserAssetDataSource/Private/ContentBrowserAssetDataCore.cpp
+++ b/Engine/Plugins/Editor/ContentBrowser/ContentBrowserAssetDataSource/Source/ContentBrowserAssetDataSource/Private/ContentBrowserAssetDataCore.cpp
@@ -148,18 +148,16 @@ bool IsPrimaryAsset(const FAssetData& InAssetData)
 {
 	// External assets are not displayed in the Content Browser or other asset pickers
 	const bool bIsExternalAsset = !InAssetData.GetOptionalOuterPathName().IsNone();
-
-	// Check for the asset being a redirector first, as currently only class 
-	// redirectors emit non-primary assets from the Asset Registry
-	return !InAssetData.IsRedirector() && InAssetData.IsUAsset() && !bIsExternalAsset;
+	
+	return !bIsExternalAsset && InAssetData.IsUAsset();
 }
 
 bool IsPrimaryAsset(UObject* InObject)
 {
 	// External assets are not displayed in the Content Browser or other asset pickers
 	const bool bIsExternalAsset = InObject->IsPackageExternal();
-
-	return !FAssetData::IsRedirector(InObject) && FAssetData::IsUAsset(InObject) && !bIsExternalAsset;
+	
+	return !bIsExternalAsset && FAssetData::IsUAsset(InObject);
 }
 
 void SetOptionalErrorMessage(FText* OutErrorMsg, FText InErrorMsg)
