Content Browser - Fix cooked assets no longer showing up

diff --git a/Engine/Plugins/Editor/ContentBrowser/ContentBrowserAssetDataSource/Source/ContentBrowserAssetDataSource/Private/ContentBrowserAssetDataCore.cpp b/Engine/Plugins/Editor/ContentBrowser/ContentBrowserAssetDataSource/Source/ContentBrowserAssetDataSource/Private/ContentBrowserAssetDataCore.cpp
--- a/Engine/Plugins/Editor/ContentBrowser/ContentBrowserAssetDataSource/Source/ContentBrowserAssetDataSource/Private/ContentBrowserAssetDataCore.cpp
+++ b/Engine/Plugins/Editor/ContentBrowser/ContentBrowserAssetDataSource/Source/ContentBrowserAssetDataSource/Private/ContentBrowserAssetDataCore.cpp
@@ -146,10 +146,22 @@ void EnumerateAssetItemPayloads(const UContentBrowserDataSource* InOwnerDataSour
 
 bool IsPrimaryAsset(const FAssetData& InAssetData)
 {
-	// External assets are not displayed in the Content Browser or other asset pickers
-	const bool bIsExternalAsset = !InAssetData.GetOptionalOuterPathName().IsNone();
-	
-	return !bIsExternalAsset && InAssetData.IsUAsset();
+	// Both GetOptionalOuterPathName and IsUAsset currently do not work on cooked assets
+	//
+	// GetOptionalOuterPathName is not serialized to the asset registry during cook
+	// IsUAsset when called on compiled blueprint class compares Name_C vs Name and returns false
+	if (InAssetData.HasAnyPackageFlags(PKG_Cooked | PKG_FilterEditorOnly))
+	{
+		// Check for the asset being a redirector first, as currently only class 
+		// redirectors emit non-primary assets from the Asset Registry
+		return !InAssetData.IsRedirector() || InAssetData.IsUAsset();
+	}
+	else
+	{
+		// External assets are not displayed in the Content Browser or other asset pickers
+		bool bIsExternalAsset = !InAssetData.GetOptionalOuterPathName().IsNone();
+		return !bIsExternalAsset && InAssetData.IsUAsset();
+	}
 }
 
 bool IsPrimaryAsset(UObject* InObject)
