Update wording of trying to edit cooked assets to be less scary for new users

diff --git a/Engine/Plugins/Editor/ContentBrowser/ContentBrowserAssetDataSource/Source/ContentBrowserAssetDataSource/Private/ContentBrowserAssetDataCore.cpp b/Engine/Plugins/Editor/ContentBrowser/ContentBrowserAssetDataSource/Source/ContentBrowserAssetDataSource/Private/ContentBrowserAssetDataCore.cpp
--- a/Engine/Plugins/Editor/ContentBrowser/ContentBrowserAssetDataSource/Source/ContentBrowserAssetDataSource/Private/ContentBrowserAssetDataCore.cpp
+++ b/Engine/Plugins/Editor/ContentBrowser/ContentBrowserAssetDataSource/Source/ContentBrowserAssetDataSource/Private/ContentBrowserAssetDataCore.cpp
@@ -185,7 +185,7 @@ bool CanModifyPath(IAssetTools* InAssetTools, const FName InFolderPath, FText* O
 	const TSharedRef<FPathPermissionList>& WritableFolderFilter = InAssetTools->GetWritableFolderPermissionList();
 	if (!WritableFolderFilter->PassesStartsWithFilter(InFolderPath))
 	{
-		SetOptionalErrorMessage(OutErrorMsg, FText::Format(LOCTEXT("Error_FolderIsLocked", "Folder '{0}' is Locked"), FText::FromName(InFolderPath)));
+		SetOptionalErrorMessage(OutErrorMsg, FText::Format(LOCTEXT("Error_FolderIsLocked", "Folder '{0}' is read only and its contents cannot be edited"), FText::FromName(InFolderPath)));
 		return false;
 	}
 	return true;
diff --git a/Engine/Source/Editor/ContentBrowser/Private/SContentBrowser.cpp b/Engine/Source/Editor/ContentBrowser/Private/SContentBrowser.cpp
--- a/Engine/Source/Editor/ContentBrowser/Private/SContentBrowser.cpp
+++ b/Engine/Source/Editor/ContentBrowser/Private/SContentBrowser.cpp
@@ -3082,7 +3082,7 @@ void SContentBrowser::OnItemsActivated(TArrayView<const FContentBrowserItem> Act
 			
 			if (!SourceAndItemsPair.Key->BulkEditItems(SourceAndItemsPair.Value))
 			{
-				static const FText ErrorMessage = LOCTEXT("EditItemsFailure", "Failed to load assets");
+				static const FText ErrorMessage = LOCTEXT("EditItemsFailure", "Unable to edit assets");
 
 				FNotificationInfo WarningNotification(ErrorMessage);
 				WarningNotification.ExpireDuration = 5.0f;
