Content Browser UX improvements for cooked content

* Plugins that only contain cooked content (and have no C++ modules) no longer show their "Content" suffix.
* Folders that only contain private cooked content are no longer visible in the asset or path views.

diff --git a/Engine/Plugins/Editor/ContentBrowser/ContentBrowserAssetDataSource/Source/ContentBrowserAssetDataSource/Private/ContentBrowserAssetDataCore.cpp b/Engine/Plugins/Editor/ContentBrowser/ContentBrowserAssetDataSource/Source/ContentBrowserAssetDataSource/Private/ContentBrowserAssetDataCore.cpp
--- a/Engine/Plugins/Editor/ContentBrowser/ContentBrowserAssetDataSource/Source/ContentBrowserAssetDataSource/Private/ContentBrowserAssetDataCore.cpp
+++ b/Engine/Plugins/Editor/ContentBrowser/ContentBrowserAssetDataSource/Source/ContentBrowserAssetDataSource/Private/ContentBrowserAssetDataCore.cpp
@@ -35,10 +35,10 @@ DEFINE_LOG_CATEGORY_STATIC(LogContentBrowserAssetDataSource, Warning, Warning);
 namespace ContentBrowserAssetData
 {
 
-FContentBrowserItemData CreateAssetFolderItem(UContentBrowserDataSource* InOwnerDataSource, const FName InVirtualPath, const FName InFolderPath)
+FContentBrowserItemData CreateAssetFolderItem(UContentBrowserDataSource* InOwnerDataSource, const FName InVirtualPath, const FName InFolderPath, const bool bIsCookedPath)
 {
 	const FString FolderItemName = FPackageName::GetShortName(InFolderPath);
-	FText FolderDisplayNameOverride = ContentBrowserDataUtils::GetFolderItemDisplayNameOverride(InFolderPath, FolderItemName, /*bIsClassesFolder*/ false);
+	FText FolderDisplayNameOverride = ContentBrowserDataUtils::GetFolderItemDisplayNameOverride(InFolderPath, FolderItemName, /*bIsClassesFolder*/ false, bIsCookedPath);
 	return FContentBrowserItemData(InOwnerDataSource, EContentBrowserItemFlags::Type_Folder | EContentBrowserItemFlags::Category_Asset, InVirtualPath, *FolderItemName, MoveTemp(FolderDisplayNameOverride), MakeShared<FContentBrowserAssetFolderItemDataPayload>(InFolderPath));
 }
 
diff --git a/Engine/Plugins/Editor/ContentBrowser/ContentBrowserAssetDataSource/Source/ContentBrowserAssetDataSource/Private/ContentBrowserAssetDataSource.cpp b/Engine/Plugins/Editor/ContentBrowser/ContentBrowserAssetDataSource/Source/ContentBrowserAssetDataSource/Private/ContentBrowserAssetDataSource.cpp
--- a/Engine/Plugins/Editor/ContentBrowser/ContentBrowserAssetDataSource/Source/ContentBrowserAssetDataSource/Private/ContentBrowserAssetDataSource.cpp
+++ b/Engine/Plugins/Editor/ContentBrowser/ContentBrowserAssetDataSource/Source/ContentBrowserAssetDataSource/Private/ContentBrowserAssetDataSource.cpp
@@ -63,7 +63,6 @@ void UContentBrowserAssetDataSource::Initialize(const bool InAutoRegister)
 	AssetRegistry->OnAssetUpdatedOnDisk().AddUObject(this, &UContentBrowserAssetDataSource::OnAssetUpdatedOnDisk);
 	AssetRegistry->OnPathAdded().AddUObject(this, &UContentBrowserAssetDataSource::OnPathAdded);
 	AssetRegistry->OnPathRemoved().AddUObject(this, &UContentBrowserAssetDataSource::OnPathRemoved);
-	AssetRegistry->OnFilesLoaded().AddUObject(this, &UContentBrowserAssetDataSource::OnScanCompleted);
 
 	// Listen for when assets are loaded or changed
 	FCoreUObjectDelegates::OnObjectPropertyChanged.AddUObject(this, &UContentBrowserAssetDataSource::OnObjectPropertyChanged);
@@ -147,17 +146,14 @@ void UContentBrowserAssetDataSource::Initialize(const bool InAutoRegister)
 
 	DiscoveryStatusText = LOCTEXT("InitializingAssetDiscovery", "Initializing Asset Discovery...");
 
-	// Populate the initial set of hidden empty folders
+	// Populate the initial set of folder attributes
 	// This will be updated as the scan finds more content
-	AssetRegistry->EnumerateAllCachedPaths([this](FName InPath)
+	AssetRegistry->EnumerateAllAssets([this](const FAssetData& InAssetData)
 		{
-			if (!AssetRegistry->HasAssets(InPath, /*bRecursive*/true))
-			{
-				EmptyAssetFolders.Add(InPath);
-			}
+			OnPathPopulated(InAssetData);
 			return true;
-		});
-	VisitedEmptyAssetFolders.Empty();
+		}, /*bIncludeOnlyOnDiskAssets*/true);
+	RecentlyPopulatedAssetFolders.Empty();
 
 	FPackageName::QueryRootContentPaths(RootContentPaths);
 
@@ -1281,7 +1277,7 @@ bool UContentBrowserAssetDataSource::PrioritizeSearchPath(const FName InPath)
 	return true;
 }
 
-bool UContentBrowserAssetDataSource::IsFolderVisibleIfHidingEmpty(const FName InPath)
+bool UContentBrowserAssetDataSource::IsFolderVisible(const FName InPath, const EContentBrowserIsFolderVisibleFlags InFlags)
 {
 	FName ConvertedPath;
 	const EContentBrowserPathType ConvertedPathType = TryConvertVirtualPath(InPath, ConvertedPath);
@@ -1301,8 +1297,20 @@ bool UContentBrowserAssetDataSource::IsFolderVisibleIfHidingEmpty(const FName In
 		return false;
 	}
 
-	return AlwaysVisibleAssetFolders.Contains(ConvertedPath)
-		|| !EmptyAssetFolders.Contains(ConvertedPath);
+	const EContentBrowserFolderAttributes FolderAttributes = GetAssetFolderAttributes(ConvertedPath);
+	if (EnumHasAnyFlags(FolderAttributes, EContentBrowserFolderAttributes::AlwaysVisible))
+	{
+		return true;
+	}
+
+	// Hide folders that only contain cooked private content
+	if (EnumHasAnyFlags(FolderAttributes, EContentBrowserFolderAttributes::HasContent) && !EnumHasAnyFlags(FolderAttributes, EContentBrowserFolderAttributes::HasPublicContent | EContentBrowserFolderAttributes::HasSourceContent))
+	{
+		return false;
+	}
+
+	return !EnumHasAnyFlags(InFlags, EContentBrowserIsFolderVisibleFlags::HideEmptyFolders)
+		|| EnumHasAnyFlags(FolderAttributes, EContentBrowserFolderAttributes::HasContent);
 }
 
 bool UContentBrowserAssetDataSource::CanCreateFolder(const FName InPath, FText* OutErrorMsg)
@@ -2050,7 +2058,10 @@ FContentBrowserItemData UContentBrowserAssetDataSource::CreateAssetFolderItem(co
 	FName VirtualizedPath;
 	TryConvertInternalPathToVirtual(InFolderPath, VirtualizedPath);
 
-	return ContentBrowserAssetData::CreateAssetFolderItem(this, VirtualizedPath, InFolderPath);
+	const EContentBrowserFolderAttributes FolderAttributes = GetAssetFolderAttributes(InFolderPath);
+	const bool bIsCookedPath = EnumHasAnyFlags(FolderAttributes, EContentBrowserFolderAttributes::HasContent) && !EnumHasAnyFlags(FolderAttributes, EContentBrowserFolderAttributes::HasSourceContent);
+
+	return ContentBrowserAssetData::CreateAssetFolderItem(this, VirtualizedPath, InFolderPath, bIsCookedPath);
 }
 
 FContentBrowserItemData UContentBrowserAssetDataSource::CreateAssetFileItem(const FAssetData& InAssetData)
@@ -2126,7 +2137,7 @@ void UContentBrowserAssetDataSource::OnAssetAdded(const FAssetData& InAssetData)
 	if (ContentBrowserAssetData::IsPrimaryAsset(InAssetData))
 	{
 		// The owner folder of this asset is no longer considered empty
-		OnPathPopulated(InAssetData.PackagePath);
+		OnPathPopulated(InAssetData);
 
 		QueueItemDataUpdate(FContentBrowserItemDataUpdate::MakeItemAddedUpdate(CreateAssetFileItem(InAssetData)));
 	}
@@ -2145,7 +2156,7 @@ void UContentBrowserAssetDataSource::OnAssetRenamed(const FAssetData& InAssetDat
 	if (ContentBrowserAssetData::IsPrimaryAsset(InAssetData))
 	{
 		// The owner folder of this asset is no longer considered empty
-		OnPathPopulated(InAssetData.PackagePath);
+		OnPathPopulated(InAssetData);
 
 		FName VirtualizedPath;
 		TryConvertInternalPathToVirtual(*InOldObjectPath, VirtualizedPath);
@@ -2190,11 +2201,9 @@ void UContentBrowserAssetDataSource::OnObjectPreSave(UObject* InObject, FObjectP
 
 void UContentBrowserAssetDataSource::OnPathAdded(const FString& InPath)
 {
-	// New paths are considered empty until assets are added inside them
 	FName PathName(InPath);
-	EmptyAssetFolders.Add(PathName);
-	VisitedEmptyAssetFolders.Empty();
-
+	RecentlyPopulatedAssetFolders.Empty();
+	
 	QueueItemDataUpdate(FContentBrowserItemDataUpdate::MakeItemAddedUpdate(CreateAssetFolderItem(PathName)));
 }
 
@@ -2202,18 +2211,23 @@ void UContentBrowserAssetDataSource::OnPathRemoved(const FString& InPath)
 {
 	// Deleted paths are no longer relevant for tracking
 	FName PathName(InPath);
-	AlwaysVisibleAssetFolders.Remove(PathName);
-	EmptyAssetFolders.Remove(PathName);
+	RecentlyPopulatedAssetFolders.Remove(PathName);
+	AssetFolderToAttributes.Remove(PathName);
 
 	QueueItemDataUpdate(FContentBrowserItemDataUpdate::MakeItemRemovedUpdate(CreateAssetFolderItem(PathName)));
 }
 
-void UContentBrowserAssetDataSource::OnPathPopulated(const FName InPath)
+void UContentBrowserAssetDataSource::OnPathPopulated(const FAssetData& InAssetData)
 {
-	OnPathPopulated(FNameBuilder(InPath));
+	const EContentBrowserFolderAttributes FolderAttributes 
+		= EContentBrowserFolderAttributes::HasContent
+		| (InAssetData.PackageFlags & PKG_Cooked ? EContentBrowserFolderAttributes::None : EContentBrowserFolderAttributes::HasSourceContent)
+		| (InAssetData.PackageFlags & PKG_NotExternallyReferenceable ? EContentBrowserFolderAttributes::None : EContentBrowserFolderAttributes::HasPublicContent);
+
+	OnPathPopulated(FNameBuilder(InAssetData.PackagePath), FolderAttributes);
 }
 
-void UContentBrowserAssetDataSource::OnPathPopulated(const FStringView InPath)
+void UContentBrowserAssetDataSource::OnPathPopulated(const FStringView InPath, const EContentBrowserFolderAttributes InAttributesToSet)
 {
 	// Recursively un-hide this path, emitting update events for any paths that change state so that the view updates
 	if (InPath.Len() > 1)
@@ -2229,29 +2243,28 @@ void UContentBrowserAssetDataSource::OnPathPopulated(const FStringView InPath)
 
 		// If we've already visited this path then we can assume we visited the parents as well
 		// and can skip visiting this path and its parents
-		if (VisitedEmptyAssetFolders.Contains(PathName))
+		if (const EContentBrowserFolderAttributes* RecentlyAddedFolderAttributesPtr = RecentlyPopulatedAssetFolders.Find(PathName);
+			RecentlyAddedFolderAttributesPtr && EnumHasAllFlags(*RecentlyAddedFolderAttributesPtr, InAttributesToSet))
 		{
 			return;
 		}
 
 		// Recurse first as we want parents to be updated before their children
+		if (int32 LastSlashIndex = INDEX_NONE;
+			Path.FindLastChar(TEXT('/'), LastSlashIndex) && LastSlashIndex > 0)
 		{
-			int32 LastSlashIndex = INDEX_NONE;
-			if (Path.FindLastChar(TEXT('/'), LastSlashIndex) && LastSlashIndex > 0)
-			{
-				OnPathPopulated(Path.Left(LastSlashIndex));
-			}
+			OnPathPopulated(Path.Left(LastSlashIndex), InAttributesToSet);
 		}
 
 		// Unhide this folder and emit a notification if required
-		if (EmptyAssetFolders.Remove(PathName) > 0)
+		if (SetAssetFolderAttributes(PathName, InAttributesToSet))
 		{
 			// Queue an update event for this path as it may have become visible in the view
 			QueueItemDataUpdate(FContentBrowserItemDataUpdate::MakeItemModifiedUpdate(CreateAssetFolderItem(PathName)));
 		}
 
 		// Mark that this path has been visited
-		VisitedEmptyAssetFolders.Add(PathName);
+		RecentlyPopulatedAssetFolders.FindOrAdd(PathName) |= InAttributesToSet;
 	}
 }
 
@@ -2268,32 +2281,22 @@ void UContentBrowserAssetDataSource::OnAlwaysShowPath(const FString& InPath)
 		}
 
 		// Recurse first as we want parents to be updated before their children
+		if (int32 LastSlashIndex = INDEX_NONE;
+			Path.FindLastChar(TEXT('/'), LastSlashIndex) && LastSlashIndex > 0)
 		{
-			int32 LastSlashIndex = INDEX_NONE;
-			if (Path.FindLastChar(TEXT('/'), LastSlashIndex) && LastSlashIndex > 0)
-			{
-				OnAlwaysShowPath(Path.Left(LastSlashIndex));
-			}
+			OnAlwaysShowPath(Path.Left(LastSlashIndex));
 		}
 
 		// Force show this folder and emit a notification if required
 		FName PathName(Path);
-		if (!AlwaysVisibleAssetFolders.Contains(PathName))
+		if (SetAssetFolderAttributes(PathName, EContentBrowserFolderAttributes::AlwaysVisible))
 		{
-			AlwaysVisibleAssetFolders.Add(PathName);
-
 			// Queue an update event for this path as it may have become visible in the view
 			QueueItemDataUpdate(FContentBrowserItemDataUpdate::MakeItemModifiedUpdate(CreateAssetFolderItem(PathName)));
 		}
 	}
 }
 
-void UContentBrowserAssetDataSource::OnScanCompleted()
-{
-	// Done finding content - compact this set as items would have been removed as assets were found
-	EmptyAssetFolders.CompactStable();
-}
-
 void UContentBrowserAssetDataSource::BuildRootPathVirtualTree() 
 {
 	Super::BuildRootPathVirtualTree();
@@ -2321,6 +2324,51 @@ void UContentBrowserAssetDataSource::OnContentPathDismounted(const FString& InAs
 	RootContentPaths.Remove(InAssetPath);
 }
 
+EContentBrowserFolderAttributes UContentBrowserAssetDataSource::GetAssetFolderAttributes(const FName InPath) const
+{
+	const EContentBrowserFolderAttributes* FolderAttributesPtr = AssetFolderToAttributes.Find(InPath);
+	return FolderAttributesPtr
+		? *FolderAttributesPtr
+		: EContentBrowserFolderAttributes::None;
+}
+
+bool UContentBrowserAssetDataSource::SetAssetFolderAttributes(const FName InPath, const EContentBrowserFolderAttributes InAttributesToSet)
+{
+	if (InAttributesToSet != EContentBrowserFolderAttributes::None)
+	{
+		EContentBrowserFolderAttributes& FolderAttributes = AssetFolderToAttributes.FindOrAdd(InPath);
+
+		const EContentBrowserFolderAttributes PreviousAttributes = FolderAttributes;
+		EnumAddFlags(FolderAttributes, InAttributesToSet);
+
+		const bool bHasChanged = FolderAttributes != PreviousAttributes;
+		return bHasChanged;
+	}
+
+	return false;
+}
+
+bool UContentBrowserAssetDataSource::ClearAssetFolderAttributes(const FName InPath, const EContentBrowserFolderAttributes InAttributesToClear)
+{
+	if (InAttributesToClear != EContentBrowserFolderAttributes::None)
+	{
+		if (EContentBrowserFolderAttributes* FolderAttributesPtr = AssetFolderToAttributes.Find(InPath))
+		{
+			const EContentBrowserFolderAttributes PreviousAttributes = *FolderAttributesPtr;
+			EnumRemoveFlags(*FolderAttributesPtr, InAttributesToClear);
+
+			const bool bHasChanged = *FolderAttributesPtr != PreviousAttributes;
+			if (*FolderAttributesPtr == EContentBrowserFolderAttributes::None)
+			{
+				AssetFolderToAttributes.Remove(InPath);
+			}
+			return bHasChanged;
+		}
+	}
+
+	return false;
+}
+
 void UContentBrowserAssetDataSource::PopulateAddNewContextMenu(UToolMenu* InMenu)
 {
 	const UContentBrowserDataMenuContext_AddNewMenu* ContextObject = InMenu->FindContext<UContentBrowserDataMenuContext_AddNewMenu>();
diff --git a/Engine/Plugins/Editor/ContentBrowser/ContentBrowserAssetDataSource/Source/ContentBrowserAssetDataSource/Public/ContentBrowserAssetDataCore.h b/Engine/Plugins/Editor/ContentBrowser/ContentBrowserAssetDataSource/Source/ContentBrowserAssetDataSource/Public/ContentBrowserAssetDataCore.h
--- a/Engine/Plugins/Editor/ContentBrowser/ContentBrowserAssetDataSource/Source/ContentBrowserAssetDataSource/Public/ContentBrowserAssetDataCore.h
+++ b/Engine/Plugins/Editor/ContentBrowser/ContentBrowserAssetDataSource/Source/ContentBrowserAssetDataSource/Public/ContentBrowserAssetDataCore.h
@@ -20,7 +20,7 @@ class UContentBrowserDataSource;
 namespace ContentBrowserAssetData
 {
 
-	CONTENTBROWSERASSETDATASOURCE_API FContentBrowserItemData CreateAssetFolderItem(UContentBrowserDataSource* InOwnerDataSource, const FName InVirtualPath, const FName InFolderPath);
+	CONTENTBROWSERASSETDATASOURCE_API FContentBrowserItemData CreateAssetFolderItem(UContentBrowserDataSource* InOwnerDataSource, const FName InVirtualPath, const FName InFolderPath, const bool bIsCookedPath = false);
 
 	CONTENTBROWSERASSETDATASOURCE_API FContentBrowserItemData CreateAssetFileItem(UContentBrowserDataSource* InOwnerDataSource, const FName InVirtualPath, const FAssetData& InAssetData);
 
diff --git a/Engine/Plugins/Editor/ContentBrowser/ContentBrowserAssetDataSource/Source/ContentBrowserAssetDataSource/Public/ContentBrowserAssetDataSource.h b/Engine/Plugins/Editor/ContentBrowser/ContentBrowserAssetDataSource/Source/ContentBrowserAssetDataSource/Public/ContentBrowserAssetDataSource.h
--- a/Engine/Plugins/Editor/ContentBrowser/ContentBrowserAssetDataSource/Source/ContentBrowserAssetDataSource/Public/ContentBrowserAssetDataSource.h
+++ b/Engine/Plugins/Editor/ContentBrowser/ContentBrowserAssetDataSource/Source/ContentBrowserAssetDataSource/Public/ContentBrowserAssetDataSource.h
@@ -75,6 +75,36 @@ public:
 	FARCompiledFilter ExclusiveFilter;
 };
 
+enum class EContentBrowserFolderAttributes : uint8
+{
+	/**
+	 * No special attributes.
+	 */
+	None = 0,
+
+	/**
+	 * This folder should always be visible, even if it contains no content in the Content Browser view.
+	 * This will include root content folders, and any folders that have been created directly (or indirectly) by a user action.
+	 */
+	AlwaysVisible = 1<<0,
+
+	/**
+	 * This folder has content that will appear in the Content Browser view.
+	 */
+	HasContent = 1<<1,
+
+	/**
+	 * This folder has public content that will appear in the Content Browser view.
+	 */
+	HasPublicContent = 1<<2,
+
+	/**
+	 * This folder has source (uncooked) content that will appear in the Content Browser view.
+	 */
+	HasSourceContent = 1<<3,
+};
+ENUM_CLASS_FLAGS(EContentBrowserFolderAttributes);
+
 UCLASS()
 class CONTENTBROWSERASSETDATASOURCE_API UContentBrowserAssetDataSource : public UContentBrowserDataSource
 {
@@ -195,7 +225,7 @@ public:
 
 	virtual bool PrioritizeSearchPath(const FName InPath) override;
 
-	virtual bool IsFolderVisibleIfHidingEmpty(const FName InPath) override;
+	virtual bool IsFolderVisible(const FName InPath, const EContentBrowserIsFolderVisibleFlags InFlags) override;
 
 	virtual bool CanCreateFolder(const FName InPath, FText* OutErrorMsg) override;
 
@@ -339,18 +369,22 @@ private:
 
 	void OnPathRemoved(const FString& InPath);
 
-	void OnPathPopulated(const FName InPath);
+	void OnPathPopulated(const FAssetData& InAssetData);
 
-	void OnPathPopulated(const FStringView InPath);
+	void OnPathPopulated(const FStringView InPath, const EContentBrowserFolderAttributes InAttributesToSet);
 
 	void OnAlwaysShowPath(const FString& InPath);
 
-	void OnScanCompleted();
-
 	void OnContentPathMounted(const FString& InAssetPath, const FString& InFileSystemPath);
 
 	void OnContentPathDismounted(const FString& InAssetPath, const FString& InFileSystemPath);
 
+	EContentBrowserFolderAttributes GetAssetFolderAttributes(const FName InPath) const;
+
+	bool SetAssetFolderAttributes(const FName InPath, const EContentBrowserFolderAttributes InAttributesToSet);
+
+	bool ClearAssetFolderAttributes(const FName InPath, const EContentBrowserFolderAttributes InAttributesToClear);
+
 	void PopulateAddNewContextMenu(UToolMenu* InMenu);
 
 	void PopulateContentBrowserToolBar(UToolMenu* InMenu);
@@ -400,20 +434,14 @@ private:
 	TArray<FString> RootContentPaths;
 
 	/**
-	 * The set of folders that should always be visible, even if they contain no assets in the Content Browser view.
-	 * This will include root content folders, and any folders that have been created directly (or indirectly) by a user action.
+	 * Map of folders that have attributes set.
 	 */
-	TSet<FName> AlwaysVisibleAssetFolders;
+	TMap<FName, EContentBrowserFolderAttributes> AssetFolderToAttributes;
 
 	/**
-	 * A cache of folders that contain no assets in the Content Browser view.
-	 */
-	TSet<FName> EmptyAssetFolders;
-	
-	/**
-	 * A cache of folders that visited since the last time any empty asset folders were added.
+	 * A cache of folders that have been populated since the last time any new asset folders were added.
 	 */
-	TSet<FName> VisitedEmptyAssetFolders;
+	TMap<FName, EContentBrowserFolderAttributes> RecentlyPopulatedAssetFolders;
 };
 
 #if UE_ENABLE_INCLUDE_ORDER_DEPRECATED_IN_5_2
diff --git a/Engine/Plugins/Editor/ContentBrowser/ContentBrowserClassDataSource/Source/ContentBrowserClassDataSource/Private/ContentBrowserClassDataSource.cpp b/Engine/Plugins/Editor/ContentBrowser/ContentBrowserClassDataSource/Source/ContentBrowserClassDataSource/Private/ContentBrowserClassDataSource.cpp
--- a/Engine/Plugins/Editor/ContentBrowser/ContentBrowserClassDataSource/Source/ContentBrowserClassDataSource/Private/ContentBrowserClassDataSource.cpp
+++ b/Engine/Plugins/Editor/ContentBrowser/ContentBrowserClassDataSource/Source/ContentBrowserClassDataSource/Private/ContentBrowserClassDataSource.cpp
@@ -412,7 +412,7 @@ bool UContentBrowserClassDataSource::EnumerateItemsForObjects(const TArrayView<U
 	return true;
 }
 
-bool UContentBrowserClassDataSource::IsFolderVisibleIfHidingEmpty(const FName InPath)
+bool UContentBrowserClassDataSource::IsFolderVisible(const FName InPath, const EContentBrowserIsFolderVisibleFlags InFlags)
 {
 	FName ConvertedPath;
 	const EContentBrowserPathType ConvertedPathType = TryConvertVirtualPath(InPath, ConvertedPath);
@@ -435,6 +435,7 @@ bool UContentBrowserClassDataSource::IsFolderVisibleIfHidingEmpty(const FName In
 	ConditionalCreateNativeClassHierarchy();
 
 	return ContentBrowserDataUtils::IsTopLevelFolder(ConvertedPath) 
+		|| !EnumHasAnyFlags(InFlags, EContentBrowserIsFolderVisibleFlags::HideEmptyFolders)
 		|| NativeClassHierarchy->HasClasses(ConvertedPath, /*bRecursive*/true);
 }
 
diff --git a/Engine/Plugins/Editor/ContentBrowser/ContentBrowserClassDataSource/Source/ContentBrowserClassDataSource/Public/ContentBrowserClassDataSource.h b/Engine/Plugins/Editor/ContentBrowser/ContentBrowserClassDataSource/Source/ContentBrowserClassDataSource/Public/ContentBrowserClassDataSource.h
--- a/Engine/Plugins/Editor/ContentBrowser/ContentBrowserClassDataSource/Source/ContentBrowserClassDataSource/Public/ContentBrowserClassDataSource.h
+++ b/Engine/Plugins/Editor/ContentBrowser/ContentBrowserClassDataSource/Source/ContentBrowserClassDataSource/Public/ContentBrowserClassDataSource.h
@@ -47,7 +47,7 @@ public:
 
 	virtual bool EnumerateItemsForObjects(const TArrayView<UObject*> InObjects, TFunctionRef<bool(FContentBrowserItemData&&)> InCallback) override;
 
-	virtual bool IsFolderVisibleIfHidingEmpty(const FName InPath) override;
+	virtual bool IsFolderVisible(const FName InPath, const EContentBrowserIsFolderVisibleFlags InFlags) override;
 
 	virtual bool DoesItemPassFilter(const FContentBrowserItemData& InItem, const FContentBrowserDataCompiledFilter& InFilter) override;
 
diff --git a/Engine/Plugins/Editor/ContentBrowser/ContentBrowserFileDataSource/Source/ContentBrowserFileDataSource/Private/ContentBrowserFileDataSource.cpp b/Engine/Plugins/Editor/ContentBrowser/ContentBrowserFileDataSource/Source/ContentBrowserFileDataSource/Private/ContentBrowserFileDataSource.cpp
--- a/Engine/Plugins/Editor/ContentBrowser/ContentBrowserFileDataSource/Source/ContentBrowserFileDataSource/Private/ContentBrowserFileDataSource.cpp
+++ b/Engine/Plugins/Editor/ContentBrowser/ContentBrowserFileDataSource/Source/ContentBrowserFileDataSource/Private/ContentBrowserFileDataSource.cpp
@@ -1055,7 +1055,7 @@ bool UContentBrowserFileDataSource::PrioritizeSearchPath(const FName InPath)
 	return true;
 }
 
-bool UContentBrowserFileDataSource::IsFolderVisibleIfHidingEmpty(const FName InPath)
+bool UContentBrowserFileDataSource::IsFolderVisible(const FName InPath, const EContentBrowserIsFolderVisibleFlags InFlags)
 {
 	FName ConvertedPath;
 	const EContentBrowserPathType ConvertedPathType = TryConvertVirtualPath(InPath, ConvertedPath);
@@ -1080,6 +1080,7 @@ bool UContentBrowserFileDataSource::IsFolderVisibleIfHidingEmpty(const FName InP
 	const uint32 InternalPathHash = GetTypeHash(InternalPathStrView);
 
 	return AlwaysVisibleFolders.ContainsByHash(InternalPathHash, InternalPathStrView)
+		|| !EnumHasAnyFlags(InFlags, EContentBrowserIsFolderVisibleFlags::HideEmptyFolders)
 		|| !EmptyFolders.ContainsByHash(InternalPathHash, InternalPathStrView);
 }
 
diff --git a/Engine/Plugins/Editor/ContentBrowser/ContentBrowserFileDataSource/Source/ContentBrowserFileDataSource/Public/ContentBrowserFileDataSource.h b/Engine/Plugins/Editor/ContentBrowser/ContentBrowserFileDataSource/Source/ContentBrowserFileDataSource/Public/ContentBrowserFileDataSource.h
--- a/Engine/Plugins/Editor/ContentBrowser/ContentBrowserFileDataSource/Source/ContentBrowserFileDataSource/Public/ContentBrowserFileDataSource.h
+++ b/Engine/Plugins/Editor/ContentBrowser/ContentBrowserFileDataSource/Source/ContentBrowserFileDataSource/Public/ContentBrowserFileDataSource.h
@@ -57,7 +57,7 @@ public:
 
 	virtual bool PrioritizeSearchPath(const FName InPath) override;
 
-	virtual bool IsFolderVisibleIfHidingEmpty(const FName InPath) override;
+	virtual bool IsFolderVisible(const FName InPath, const EContentBrowserIsFolderVisibleFlags InFlags) override;
 
 	virtual bool DoesItemPassFilter(const FContentBrowserItemData& InItem, const FContentBrowserDataCompiledFilter& InFilter) override;
 
diff --git a/Engine/Source/Editor/ContentBrowser/Private/ContentBrowserUtils.cpp b/Engine/Source/Editor/ContentBrowser/Private/ContentBrowserUtils.cpp
--- a/Engine/Source/Editor/ContentBrowser/Private/ContentBrowserUtils.cpp
+++ b/Engine/Source/Editor/ContentBrowser/Private/ContentBrowserUtils.cpp
@@ -718,6 +718,11 @@ FName ContentBrowserUtils::GetInvariantPath(const FContentBrowserItemPath& ItemP
 	return ItemPath.GetInternalPathName();
 }
 
+EContentBrowserIsFolderVisibleFlags ContentBrowserUtils::GetIsFolderVisibleFlags(const bool bDisplayEmpty)
+{
+	return EContentBrowserIsFolderVisibleFlags::Default | (bDisplayEmpty ? EContentBrowserIsFolderVisibleFlags::None : EContentBrowserIsFolderVisibleFlags::HideEmptyFolders);
+}
+
 bool ContentBrowserUtils::IsFavoriteFolder(const FString& FolderPath)
 {
 	return IsFavoriteFolder(FContentBrowserItemPath(FolderPath, EContentBrowserPathType::Virtual));
diff --git a/Engine/Source/Editor/ContentBrowser/Private/ContentBrowserUtils.h b/Engine/Source/Editor/ContentBrowser/Private/ContentBrowserUtils.h
--- a/Engine/Source/Editor/ContentBrowser/Private/ContentBrowserUtils.h
+++ b/Engine/Source/Editor/ContentBrowser/Private/ContentBrowserUtils.h
@@ -25,6 +25,7 @@ struct FAssetData;
 struct FContentBrowserDataFilter;
 struct FContentBrowserItem;
 struct FContentBrowserItemPath;
+enum class EContentBrowserIsFolderVisibleFlags : uint8;
 
 namespace ContentBrowserUtils
 {
@@ -98,6 +99,9 @@ namespace ContentBrowserUtils
 	/** Returns internal path if it has one, otherwise strips /All prefix from virtual path*/
 	FName GetInvariantPath(const FContentBrowserItemPath& ItemPath);
 
+	/** Get the set of flags to use with IsFolderVisible */
+	EContentBrowserIsFolderVisibleFlags GetIsFolderVisibleFlags(const bool bDisplayEmpty);
+
 	/** Returns if this folder has been marked as a favorite folder */
 	UE_DEPRECATED(5.3, "Use function that takes FContentBrowserItemPath instead.")
 	bool IsFavoriteFolder(const FString& FolderPath);
diff --git a/Engine/Source/Editor/ContentBrowser/Private/SAssetView.cpp b/Engine/Source/Editor/ContentBrowser/Private/SAssetView.cpp
--- a/Engine/Source/Editor/ContentBrowser/Private/SAssetView.cpp
+++ b/Engine/Source/Editor/ContentBrowser/Private/SAssetView.cpp
@@ -146,7 +146,7 @@ public:
 		// Folders are only subject to "empty" filtering
 		if (InItemToFilter->IsFolder())
 		{
-			return bDisplayEmptyFolders || ContentBrowserData->IsFolderVisibleIfHidingEmpty(InItemToFilter->GetItem().GetVirtualPath());
+			return ContentBrowserData->IsFolderVisible(InItemToFilter->GetItem().GetVirtualPath(), ContentBrowserUtils::GetIsFolderVisibleFlags(bDisplayEmptyFolders));
 		}
 
 		// Run the item through the filters
diff --git a/Engine/Source/Editor/ContentBrowser/Private/SContentBrowser.cpp b/Engine/Source/Editor/ContentBrowser/Private/SContentBrowser.cpp
--- a/Engine/Source/Editor/ContentBrowser/Private/SContentBrowser.cpp
+++ b/Engine/Source/Editor/ContentBrowser/Private/SContentBrowser.cpp
@@ -2553,7 +2553,7 @@ TSharedRef<SWidget> SContentBrowser::OnGetCrumbDelimiterContent(const FString& C
 				}
 			}
 
-			if (!bDisplayEmpty && !ContentBrowserData->IsFolderVisibleIfHidingEmpty(Item.GetVirtualPath()))
+			if (!ContentBrowserData->IsFolderVisible(Item.GetVirtualPath(), ContentBrowserUtils::GetIsFolderVisibleFlags(bDisplayEmpty)))
 			{
 				It.RemoveCurrent();
 				continue;
diff --git a/Engine/Source/Editor/ContentBrowser/Private/SPathView.cpp b/Engine/Source/Editor/ContentBrowser/Private/SPathView.cpp
--- a/Engine/Source/Editor/ContentBrowser/Private/SPathView.cpp
+++ b/Engine/Source/Editor/ContentBrowser/Private/SPathView.cpp
@@ -1551,7 +1551,7 @@ void SPathView::Populate(const bool bIsRefreshingFilter)
 		TArray<TSharedPtr<FTreeItem>> ItemsCreated;
 		ContentBrowserData->EnumerateItemsMatchingFilter(CompiledDataFilter, [this, bFilteringByText, bDisplayEmpty, ContentBrowserData, &ItemsCreated](FContentBrowserItemData&& InItemData)
 		{
-			bool bPassesFilter = bDisplayEmpty || ContentBrowserData->IsFolderVisibleIfHidingEmpty(InItemData.GetVirtualPath());
+			bool bPassesFilter = ContentBrowserData->IsFolderVisible(InItemData.GetVirtualPath(), ContentBrowserUtils::GetIsFolderVisibleFlags(bDisplayEmpty));
 			if (bPassesFilter && bFilteringByText)
 			{
 				// Use the whole path so we deliberately include any children of matched parents in the filtered list
@@ -2012,7 +2012,7 @@ void SPathView::HandleItemDataUpdated(TArrayView<const FContentBrowserItemDataUp
 			return false;
 		}
 
-		if (!bDisplayEmpty && !ContentBrowserData->IsFolderVisibleIfHidingEmpty(InItemData.GetVirtualPath()))
+		if (!ContentBrowserData->IsFolderVisible(InItemData.GetVirtualPath(), ContentBrowserUtils::GetIsFolderVisibleFlags(bDisplayEmpty)))
 		{
 			return false;
 		}
diff --git a/Engine/Source/Editor/ContentBrowserData/Private/ContentBrowserDataSource.cpp b/Engine/Source/Editor/ContentBrowserData/Private/ContentBrowserDataSource.cpp
--- a/Engine/Source/Editor/ContentBrowserData/Private/ContentBrowserDataSource.cpp
+++ b/Engine/Source/Editor/ContentBrowserData/Private/ContentBrowserDataSource.cpp
@@ -245,7 +245,7 @@ bool UContentBrowserDataSource::PrioritizeSearchPath(const FName InPath)
 	return false;
 }
 
-bool UContentBrowserDataSource::IsFolderVisibleIfHidingEmpty(const FName InPath)
+bool UContentBrowserDataSource::IsFolderVisible(const FName InPath, const EContentBrowserIsFolderVisibleFlags InFlags)
 {
 	return true;
 }
diff --git a/Engine/Source/Editor/ContentBrowserData/Private/ContentBrowserDataSubsystem.cpp b/Engine/Source/Editor/ContentBrowserData/Private/ContentBrowserDataSubsystem.cpp
--- a/Engine/Source/Editor/ContentBrowserData/Private/ContentBrowserDataSubsystem.cpp
+++ b/Engine/Source/Editor/ContentBrowserData/Private/ContentBrowserDataSubsystem.cpp
@@ -537,7 +537,7 @@ bool UContentBrowserDataSubsystem::PrioritizeSearchPath(const FName InPath)
 	return bDidPrioritize;
 }
 
-bool UContentBrowserDataSubsystem::IsFolderVisibleIfHidingEmpty(const FName InPath) const
+bool UContentBrowserDataSubsystem::IsFolderVisible(const FName InPath, const EContentBrowserIsFolderVisibleFlags InFlags) const
 {
 	bool bIsVisible = false;
 	bool bIsKnownPath = false;
@@ -547,12 +547,17 @@ bool UContentBrowserDataSubsystem::IsFolderVisibleIfHidingEmpty(const FName InPa
 		if (DataSource->IsVirtualPathUnderMountRoot(InPath))
 		{
 			bIsKnownPath = true;
-			bIsVisible |= DataSource->IsFolderVisibleIfHidingEmpty(InPath);
+			bIsVisible |= DataSource->IsFolderVisible(InPath, InFlags);
 		}
 	}
 	return bIsVisible || !bIsKnownPath;
 }
 
+bool UContentBrowserDataSubsystem::IsFolderVisibleIfHidingEmpty(const FName InPath) const
+{
+	return IsFolderVisible(InPath, EContentBrowserIsFolderVisibleFlags::Default | EContentBrowserIsFolderVisibleFlags::HideEmptyFolders);
+}
+
 bool UContentBrowserDataSubsystem::CanCreateFolder(const FName InPath, FText* OutErrorMsg) const
 {
 	bool bCanCreateFolder = false;
diff --git a/Engine/Source/Editor/ContentBrowserData/Private/ContentBrowserDataUtils.cpp b/Engine/Source/Editor/ContentBrowserData/Private/ContentBrowserDataUtils.cpp
--- a/Engine/Source/Editor/ContentBrowserData/Private/ContentBrowserDataUtils.cpp
+++ b/Engine/Source/Editor/ContentBrowserData/Private/ContentBrowserDataUtils.cpp
@@ -187,7 +187,7 @@ bool ContentBrowserDataUtils::PathPassesAttributeFilter(const FStringView InPath
 	return true;
 }
 
-FText ContentBrowserDataUtils::GetFolderItemDisplayNameOverride(const FName InFolderPath, const FString& InFolderItemName, const bool bIsClassesFolder)
+FText ContentBrowserDataUtils::GetFolderItemDisplayNameOverride(const FName InFolderPath, const FString& InFolderItemName, const bool bIsClassesFolder, const bool bIsCookedPath)
 {
 	FText FolderDisplayNameOverride;
 
@@ -228,18 +228,13 @@ FText ContentBrowserDataUtils::GetFolderItemDisplayNameOverride(const FName InFo
 				}
 			}
 
+			TSharedPtr<IPlugin> Plugin = IPluginManager::Get().FindPlugin(TopLevelFolderName);
+
 			FString OverrideName;
-			if (GetDefault<UContentBrowserSettings>()->bDisplayFriendlyNameForPluginFolders)
+			if (GetDefault<UContentBrowserSettings>()->bDisplayFriendlyNameForPluginFolders && Plugin && Plugin->GetFriendlyName().Len() > 0)
 			{
-				if (TSharedPtr<IPlugin> Plugin = IPluginManager::Get().FindPlugin(TopLevelFolderName))
-				{
-					if (Plugin->GetFriendlyName().Len() > 0)
-					{
-						OverrideName = Plugin->GetFriendlyName();
-					}
-				}
+				OverrideName = Plugin->GetFriendlyName();
 			}
-
 			if (OverrideName.IsEmpty())
 			{
 				OverrideName = TopLevelFolderName;
@@ -251,7 +246,14 @@ FText ContentBrowserDataUtils::GetFolderItemDisplayNameOverride(const FName InFo
 			}
 			else
 			{
-				if (GetDefault<UContentBrowserSettings>()->bDisplayContentFolderSuffix)
+				bool bDisplayContentFolderSuffix = GetDefault<UContentBrowserSettings>()->bDisplayContentFolderSuffix;
+				if (bDisplayContentFolderSuffix && Plugin && Plugin->GetDescriptor().Modules.Num() == 0 && bIsCookedPath)
+				{
+					// Exclude the content suffix for plugins that only contain cooked content and have no C++ modules
+					bDisplayContentFolderSuffix = false;
+				}
+
+				if (bDisplayContentFolderSuffix)
 				{
 					FolderDisplayNameOverride = FText::Format(LOCTEXT("ContentFolderDisplayNameFmt", "{0} Content"), FText::AsCultureInvariant(OverrideName));
 				}
diff --git a/Engine/Source/Editor/ContentBrowserData/Public/ContentBrowserDataSource.h b/Engine/Source/Editor/ContentBrowserData/Public/ContentBrowserDataSource.h
--- a/Engine/Source/Editor/ContentBrowserData/Public/ContentBrowserDataSource.h
+++ b/Engine/Source/Editor/ContentBrowserData/Public/ContentBrowserDataSource.h
@@ -247,11 +247,10 @@ public:
 	virtual bool PrioritizeSearchPath(const FName InPath);
 
 	/**
-	 * Query whether the given virtual folder should be visible if the UI is asking to hide empty content folders.
+	 * Query whether the given virtual folder should be visible in the UI.
 	 * @note This function must be able to answer the question quickly or not at all (and assume visible). It *must not* block doing something like a file system scan.
-	 * @note "Empty" in this case means that it recursively contains no file items.
 	 */
-	virtual bool IsFolderVisibleIfHidingEmpty(const FName InPath);
+	virtual bool IsFolderVisible(const FName InPath, const EContentBrowserIsFolderVisibleFlags InFlags);
 
 	/*
 	 * Query whether a folder can be created at the given virtual path, optionally providing error information if it cannot.
diff --git a/Engine/Source/Editor/ContentBrowserData/Public/ContentBrowserDataSubsystem.h b/Engine/Source/Editor/ContentBrowserData/Public/ContentBrowserDataSubsystem.h
--- a/Engine/Source/Editor/ContentBrowserData/Public/ContentBrowserDataSubsystem.h
+++ b/Engine/Source/Editor/ContentBrowserData/Public/ContentBrowserDataSubsystem.h
@@ -48,6 +48,22 @@ enum class EContentBrowserPathType : uint8
 	Virtual
 };
 
+enum class EContentBrowserIsFolderVisibleFlags : uint8
+{
+	None = 0,
+
+	/**
+	 * Hide folders that recursively contain no file items.
+	 */
+	HideEmptyFolders = 1<<0,
+
+	/**
+	 * Default visibility flags.
+	 */
+	Default = None,
+};
+ENUM_CLASS_FLAGS(EContentBrowserIsFolderVisibleFlags);
+
 /** Called for incremental item data updates from data sources that can provide delta-updates */
 DECLARE_MULTICAST_DELEGATE_OneParam(FOnContentBrowserItemDataUpdated, TArrayView<const FContentBrowserItemDataUpdate>);
 
@@ -235,9 +251,15 @@ public:
 	 */
 	bool PrioritizeSearchPath(const FName InPath);
 
+	/**
+	 * Query whether the given virtual folder should be visible in the UI.
+	 */
+	bool IsFolderVisible(const FName InPath, const EContentBrowserIsFolderVisibleFlags InFlags = EContentBrowserIsFolderVisibleFlags::Default) const;
+
 	/**
 	 * Query whether the given virtual folder should be visible if the UI is asking to hide empty content folders.
 	 */
+	UE_DEPRECATED(5.3, "IsFolderVisibleIfHidingEmpty is deprecated. Use IsFolderVisible instead and add EContentBrowserIsFolderVisibleFlags::HideEmptyFolders to the flags.")
 	bool IsFolderVisibleIfHidingEmpty(const FName InPath) const;
 
 	/*
diff --git a/Engine/Source/Editor/ContentBrowserData/Public/ContentBrowserDataUtils.h b/Engine/Source/Editor/ContentBrowserData/Public/ContentBrowserDataUtils.h
--- a/Engine/Source/Editor/ContentBrowserData/Public/ContentBrowserDataUtils.h
+++ b/Engine/Source/Editor/ContentBrowserData/Public/ContentBrowserDataUtils.h
@@ -38,8 +38,9 @@ namespace ContentBrowserDataUtils
 	 * @param InFolderPath Internal path to get display name override for
 	 * @param InFolderItemName Short name of InFolderPath (rightmost folder name)
 	 * @param bIsClassesFolder True if this folder is a classes folder
+	 * @param bIsCookedPath True if this folder only contains cooked content (recursively), or false if it contains any uncooked content
 	 * 
 	 * @return Override display name or empty string
 	 */
-	CONTENTBROWSERDATA_API FText GetFolderItemDisplayNameOverride(const FName InFolderPath, const FString& InFolderItemName, const bool bIsClassesFolder);
+	CONTENTBROWSERDATA_API FText GetFolderItemDisplayNameOverride(const FName InFolderPath, const FString& InFolderItemName, const bool bIsClassesFolder, const bool bIsCookedPath = false);
 }
