Speed up the creation of the new asset content menu

#rnx

- Profiling showed 60ms->10ms, improving the responsiveness.
- Calling IAssetTools::GetNewAssetFactories is relatively expensive, so I added an overload of FindFactoriesInCategory that can have a TArray of factories passed to it and then calling IAssetTools::GetNewAssetFactories only once, rather than one time per category as was being done

diff --git a/Engine/Plugins/Editor/ContentBrowser/ContentBrowserAssetDataSource/Source/ContentBrowserAssetDataSource/Private/NewAssetContextMenu.cpp b/Engine/Plugins/Editor/ContentBrowser/ContentBrowserAssetDataSource/Source/ContentBrowserAssetDataSource/Private/NewAssetContextMenu.cpp
--- a/Engine/Plugins/Editor/ContentBrowser/ContentBrowserAssetDataSource/Source/ContentBrowserAssetDataSource/Private/NewAssetContextMenu.cpp
+++ b/Engine/Plugins/Editor/ContentBrowser/ContentBrowserAssetDataSource/Source/ContentBrowserAssetDataSource/Private/NewAssetContextMenu.cpp
@@ -62,15 +62,30 @@ struct FCategorySubMenuItem
 	}
 };
 
-TArray<FFactoryItem> FindFactoriesInCategory(EAssetTypeCategories::Type AssetTypeCategory, bool FindFirstOnly)
+/** Utility to return the new asset factories from FAssetToolsModule */
+static const TArray<UFactory*> GetNewAssetFactories()
 {
-    QUICK_SCOPE_CYCLE_COUNTER(FindFactoriesInCategory);
-	
-	TArray<FFactoryItem> FactoriesInThisCategory;
+	QUICK_SCOPE_CYCLE_COUNTER(GetNewAssetFactories);
 
 	static const FName NAME_AssetTools = "AssetTools";
 	const IAssetTools& AssetTools = FModuleManager::LoadModuleChecked<FAssetToolsModule>(NAME_AssetTools).Get();
-	TArray<UFactory*> Factories = AssetTools.GetNewAssetFactories();
+
+	return AssetTools.GetNewAssetFactories();
+}
+
+/**
+ * Utility to find the factories (from the set provided by the caller) with a given category.
+ * 
+ * @param Factories			The factories to look in
+ * @param AssetTypeCategory	The category to find factories for
+ * @param FindFirstOnly		Returns once the first factory has been found
+ */
+static TArray<FFactoryItem> FindFactoriesInCategory(const TArray<UFactory*>& Factories, EAssetTypeCategories::Type AssetTypeCategory, bool FindFirstOnly)
+{
+	QUICK_SCOPE_CYCLE_COUNTER(FindFactoriesInCategory);
+	
+	TArray<FFactoryItem> FactoriesInThisCategory;
+
 	for (UFactory* Factory : Factories)
 	{
 		QUICK_SCOPE_CYCLE_COUNTER(GetMenuCategories);
@@ -90,6 +105,17 @@ TArray<FFactoryItem> FindFactoriesInCategory(EAssetTypeCategories::Type AssetTyp
 	return FactoriesInThisCategory;
 }
 
+/** 
+ * Utility to find the new assert factories with a given category.
+ * 
+ * @param AssetTypeCategory	The category to find factories for
+ * @param FindFirstOnly		Returns once the first factory has been found
+ */
+static TArray<FFactoryItem> FindFactoriesInCategory(EAssetTypeCategories::Type AssetTypeCategory, bool FindFirstOnly)
+{	
+	return FindFactoriesInCategory(GetNewAssetFactories(), AssetTypeCategory, FindFirstOnly);
+}
+
 class SFactoryMenuEntry : public SCompoundWidget
 {
 public:
@@ -283,10 +309,13 @@ void FNewAssetContextMenu::MakeContextMenu(
 				return (A.CategoryName.CompareToCaseIgnored(B.CategoryName) < 0);
 			});
 
+			const IAssetTools& AssetTools = AssetToolsModule.Get();
+			const TArray<UFactory*> NewAssetFactories = AssetTools.GetNewAssetFactories();
+
 			for (const FAdvancedAssetCategory& AdvancedAssetCategory : AdvancedAssetCategories)
 			{
 				const bool FindFirstOnly = true;
-				TArray<FFactoryItem> Factories = FindFactoriesInCategory(AdvancedAssetCategory.CategoryType, FindFirstOnly);
+				TArray<FFactoryItem> Factories = FindFactoriesInCategory(NewAssetFactories, AdvancedAssetCategory.CategoryType, FindFirstOnly);
 				if (Factories.Num() > 0)
 				{
 					Section.AddSubMenu(
