Perf improvement to the function UContentBrowserAssetDataSource::CreateAssetFilter.
For a large project with a lot paths with a recursive search on all the folder the change reduced time it took to create the filter by ~86%.

Before the change a function call took 191.3 ms
After the change the same function call took 25.8 ms

diff --git a/Engine/Plugins/Editor/ContentBrowser/ContentBrowserAssetDataSource/Source/ContentBrowserAssetDataSource/Private/ContentBrowserAssetDataSource.cpp b/Engine/Plugins/Editor/ContentBrowser/ContentBrowserAssetDataSource/Source/ContentBrowserAssetDataSource/Private/ContentBrowserAssetDataSource.cpp
--- a/Engine/Plugins/Editor/ContentBrowser/ContentBrowserAssetDataSource/Source/ContentBrowserAssetDataSource/Private/ContentBrowserAssetDataSource.cpp
+++ b/Engine/Plugins/Editor/ContentBrowser/ContentBrowserAssetDataSource/Source/ContentBrowserAssetDataSource/Private/ContentBrowserAssetDataSource.cpp
@@ -518,25 +518,90 @@ PRAGMA_ENABLE_DEPRECATION_WARNINGS
 		{
 			FARCompiledFilter CompiledInternalPathFilter;
 			{
-				FARFilter InternalPathFilter;
-				for (const FName InternalPath : Params.InternalPaths)
-				{
-					InternalPathFilter.PackagePaths.Add(InternalPath);
-				}
-				InternalPathFilter.bRecursivePaths = InFilter.bRecursivePaths;
-				CreateCompiledFilter(InternalPathFilter, CompiledInternalPathFilter);
+				/**
+				 * This filter is created by testing the paths while we are recursively going down the path hierarchy
+				 * This effective because it can stop exploring a sub part of the path three when the current path fail the attribute filter
+				 * Also after a certain depth it stop testing the paths against the attribute filter since the result will the same as the parent path
+				 */
+				TRACE_CPUPROFILER_EVENT_SCOPE(UContentBrowserAssetDataSource::CreateAssetFilter::CreateInternalPathFilter);
+			
+				CompiledInternalPathFilter.PackagePaths.Reserve(Params.InternalPaths.Num());
+
+				const int32 MaxDepthPathTestNeeded = ContentBrowserDataUtils::GetMaxFolderDepthRequiredForAttributeFilter();
+
+				// This builder is shared across calls because it is stack allocated and it could cause some issues in depth recursive calls
+				FNameBuilder PathBufferStr;
+				TFunction<void (FName, int32)> TestAndGatherChildPaths;
+				TestAndGatherChildPaths = [&InFilter, &CompiledInternalPathFilter, &TestAndGatherChildPaths, &Params, &PathBufferStr, MaxDepthPathTestNeeded](FName ChildPath, int32 CurrentDepth)
+					{
+						PathBufferStr.Reset();
+						ChildPath.AppendString(PathBufferStr);
+						if (ContentBrowserDataUtils::PathPassesAttributeFilter(PathBufferStr, CurrentDepth, InFilter.ItemAttributeFilter))
+						{
+							CompiledInternalPathFilter.PackagePaths.Add(ChildPath);
+							++CurrentDepth;
+
+							if (CurrentDepth < MaxDepthPathTestNeeded)
+							{
+								constexpr bool bIsRecursive = false;
+								Params.AssetRegistry->EnumerateSubPaths(ChildPath
+									, [&TestAndGatherChildPaths, CurrentDepth](FName ChildPath)
+									{
+										TestAndGatherChildPaths(ChildPath, CurrentDepth);
+										return true;
+									}
+								, bIsRecursive);
+							}
+							else
+							{
+								constexpr bool bIsRecursive = true;
+								Params.AssetRegistry->EnumerateSubPaths(ChildPath
+									, [&CompiledInternalPathFilter](FName ChildPath)
+									{
+										CompiledInternalPathFilter.PackagePaths.Add(ChildPath);
+										return true;
+									}
+								, bIsRecursive);
+							}
+						}
+					};
 
+
+				for (const FName InternalPath : Params.InternalPaths)
 				{
-					TRACE_CPUPROFILER_EVENT_SCOPE(UContentBrowserAssetDataSource::CreateAssetFilter::RemovePaths);
-			
-					// Remove paths that do not pass item attribute filter (Engine, Plugins, Developer, Localized, __ExternalActors__ etc..)
-					for (auto It = CompiledInternalPathFilter.PackagePaths.CreateIterator(); It; ++It)
+					PathBufferStr.Reset();
+					InternalPath.AppendString(PathBufferStr);
+					FStringView Path(PathBufferStr);
+					if (ContentBrowserDataUtils::PathPassesAttributeFilter(Path, 0, InFilter.ItemAttributeFilter))
 					{
-						FNameBuilder PathStr(*It);
-						FStringView Path(PathStr);
-						if (!ContentBrowserDataUtils::PathPassesAttributeFilter(Path, 0, InFilter.ItemAttributeFilter))
+						CompiledInternalPathFilter.PackagePaths.Add(InternalPath);
+
+						if (InFilter.bRecursivePaths)
 						{
-							It.RemoveCurrent();
+							// Minus one because the test depth start at zero
+							const int32 CurrentDepth = ContentBrowserDataUtils::CalculateFolderDepthOfPath(Path) - 1;
+							if (CurrentDepth < MaxDepthPathTestNeeded)
+							{
+								constexpr bool bIsRecursive = false;
+								Params.AssetRegistry->EnumerateSubPaths(InternalPath
+									, [&TestAndGatherChildPaths, CurrentDepth](FName ChildPath)
+									{
+										TestAndGatherChildPaths(ChildPath, CurrentDepth);
+										return true;
+									}
+								, bIsRecursive);
+							}
+							else
+							{
+								constexpr bool bIsRecursive = true;
+								Params.AssetRegistry->EnumerateSubPaths(InternalPath
+									, [&CompiledInternalPathFilter](FName ChildPath)
+									{
+										CompiledInternalPathFilter.PackagePaths.Add(ChildPath);
+										return true;
+									}
+								, bIsRecursive);
+							}
 						}
 					}
 				}
diff --git a/Engine/Source/Editor/ContentBrowserData/Private/ContentBrowserDataUtils.cpp b/Engine/Source/Editor/ContentBrowserData/Private/ContentBrowserDataUtils.cpp
--- a/Engine/Source/Editor/ContentBrowserData/Private/ContentBrowserDataUtils.cpp
+++ b/Engine/Source/Editor/ContentBrowserData/Private/ContentBrowserDataUtils.cpp
@@ -61,6 +61,15 @@ bool ContentBrowserDataUtils::IsTopLevelFolder(const FName InFolderPath)
 	return IsTopLevelFolder(FNameBuilder(InFolderPath));
 }
 
+int32 ContentBrowserDataUtils::GetMaxFolderDepthRequiredForAttributeFilter()
+{
+	static const int32 MaxFolderDepthToCheck = FMath::Max(
+			ContentBrowserDataUtils::CalculateFolderDepthOfPath(FPackageName::FilenameToLongPackageName(FPaths::GameDevelopersDir()).LeftChop(1))
+		, 2);
+
+	return MaxFolderDepthToCheck;
+}
+ 
 bool ContentBrowserDataUtils::PathPassesAttributeFilter(const FStringView InPath, const int32 InAlreadyCheckedDepth, const EContentBrowserItemAttributeFilter InAttributeFilter)
 {
 	static const FString ProjectContentRootName = TEXT("Game");
@@ -69,8 +78,8 @@ bool ContentBrowserDataUtils::PathPassesAttributeFilter(const FStringView InPath
 	static const FString ExternalActorsFolderName = FPackagePath::GetExternalActorsFolderName();
 	static const FString ExternalObjectsFolderName = FPackagePath::GetExternalObjectsFolderName();
 	static const FString DeveloperPathWithoutSlash = FPackageName::FilenameToLongPackageName(FPaths::GameDevelopersDir()).LeftChop(1);
-	static int32 DevelopersFolderDepth = ContentBrowserDataUtils::CalculateFolderDepthOfPath(DeveloperPathWithoutSlash);
-	static int32 MaxFolderDepthToCheck = FMath::Max(DevelopersFolderDepth, 2);
+	static const int32 DevelopersFolderDepth = ContentBrowserDataUtils::CalculateFolderDepthOfPath(DeveloperPathWithoutSlash);
+	const int32 MaxFolderDepthToCheck = ContentBrowserDataUtils::GetMaxFolderDepthRequiredForAttributeFilter();
 
 	static auto GetRootFolderNameFromPath = [](const FStringView InFullPath)
 	{
diff --git a/Engine/Source/Editor/ContentBrowserData/Public/ContentBrowserDataUtils.h b/Engine/Source/Editor/ContentBrowserData/Public/ContentBrowserDataUtils.h
--- a/Engine/Source/Editor/ContentBrowserData/Public/ContentBrowserDataUtils.h
+++ b/Engine/Source/Editor/ContentBrowserData/Public/ContentBrowserDataUtils.h
@@ -21,6 +21,9 @@ namespace ContentBrowserDataUtils
 	/** Returns true if folder has a depth of 1 */
 	CONTENTBROWSERDATA_API bool IsTopLevelFolder(const FName InFolderPath);
 
+	/** Return the test depth after which it is not needed to test the attribute filter if the parent folder was tested */
+	CONTENTBROWSERDATA_API int32 GetMaxFolderDepthRequiredForAttributeFilter();
+
 	/**
 	 * Tests internal path against attribute filter
 	 * 
