Additional improvement for the filter compilations of the asset views.
Added a system to cache some data that can be used when recompiling the filters for the incremental data updates.
For now this only used by for the asset data source internal path compilation of the asset views.

In my test on a bad scenario this yielded a reduction 98% of the time it take to call UContentBrowserAssetDataSource::CreateAssetFilter when the cache is used.

Before the optimization the reference call took  26.3 ms
After the optimization the reference call took 0.3725 ms

Also this change fix an issue with my last change interacting badly with the alias data source by not using its customization of the package paths hierarchy.

diff --git a/Engine/Plugins/Editor/ContentBrowser/ContentBrowserAliasDataSoure/Source/ContentBrowserAliasDataSource/Private/ContentBrowserAliasDataSource.cpp b/Engine/Plugins/Editor/ContentBrowser/ContentBrowserAliasDataSoure/Source/ContentBrowserAliasDataSource/Private/ContentBrowserAliasDataSource.cpp
--- a/Engine/Plugins/Editor/ContentBrowser/ContentBrowserAliasDataSoure/Source/ContentBrowserAliasDataSource/Private/ContentBrowserAliasDataSource.cpp
+++ b/Engine/Plugins/Editor/ContentBrowser/ContentBrowserAliasDataSoure/Source/ContentBrowserAliasDataSource/Private/ContentBrowserAliasDataSource.cpp
@@ -4,7 +4,6 @@
 
 #include "AssetRegistry/AssetRegistryModule.h"
 #include "ContentBrowserAssetDataCore.h"
-#include "ContentBrowserAssetDataSource.h"
 #include "ContentBrowserDataSubsystem.h"
 #include "ContentBrowserItemPath.h"
 
@@ -67,7 +66,7 @@ void UContentBrowserAliasDataSource::CompileFilter(const FName InPath, const FCo
 	TRACE_CPUPROFILER_EVENT_SCOPE(UContentBrowserAliasDataSource::CompileFilter);
 
 	UContentBrowserAssetDataSource::FAssetFilterInputParams Params;
-	if (UContentBrowserAssetDataSource::PopulateAssetFilterInputParams(Params, this, AssetRegistry, InFilter, OutCompiledFilter, &FCollectionManagerModule::GetModule().Get()))
+	if (UContentBrowserAssetDataSource::PopulateAssetFilterInputParams(Params, this, AssetRegistry, InFilter, OutCompiledFilter, &FCollectionManagerModule::GetModule().Get(), &FilterCache))
 	{
 		// Use the DataSource's custom PathTree instead of the AssetRegistry
 		const bool bCreatedPathFilter = UContentBrowserAssetDataSource::CreatePathFilter(Params, InPath, InFilter, OutCompiledFilter, [this](FName Path, TFunctionRef<bool(FName)> Callback, bool bRecursive)
@@ -77,30 +76,16 @@ void UContentBrowserAliasDataSource::CompileFilter(const FName InPath, const FCo
 
 		if (bCreatedPathFilter)
 		{
-			UContentBrowserAssetDataSource::CreateAssetFilter(Params, InPath, InFilter, OutCompiledFilter, [this, &Params](FARFilter& InputFilter, FARCompiledFilter& OutputFilter)
-			{
-				// Same as CreatePathFilter - CompileFilter calls EnumerateSubPaths internally so this needs to intercept
-				// the filter compilation and use its own PathTree to generate the sub paths.
-				if (InputFilter.bRecursivePaths)
-				{
-					TArray<FName> PackagePaths = MoveTemp(InputFilter.PackagePaths);
-					AssetRegistry->CompileFilter(InputFilter, OutputFilter);
-
-					for (const FName Path : PackagePaths)
-					{
-						PathTree.EnumerateSubPaths(Path, [&OutputFilter](FName SubPath)
-						{
-							OutputFilter.PackagePaths.Add(SubPath);
-							return true;
-						});
-					}
-					OutputFilter.PackagePaths.Append(MoveTemp(PackagePaths));
-				}
-				else
+			auto CustomSubPathEnumeration = [this, &Params](FName Path, TFunctionRef<bool(FName)> Callback, bool bRecursive)
 				{
-					AssetRegistry->CompileFilter(InputFilter, OutputFilter);
-				}
-			});
+					// Same as CreatePathFilter - CompileFilter calls EnumerateSubPaths internally so this needs to intercept
+					// the filter compilation and use its own PathTree to generate the sub paths.
+					AssetRegistry->EnumerateSubPaths(Path, Callback,bRecursive);
+					PathTree.EnumerateSubPaths(Path, Callback, bRecursive);
+				};
+
+			UContentBrowserAssetDataSource::FSubPathEnumerationFunc CustomSubPathEnumerationRef = CustomSubPathEnumeration;
+			UContentBrowserAssetDataSource::CreateAssetFilter(Params, InPath, InFilter, OutCompiledFilter, &CustomSubPathEnumerationRef);
 		}
 	}
 }
@@ -840,6 +825,16 @@ PRAGMA_DISABLE_DEPRECATION_WARNINGS
 PRAGMA_ENABLE_DEPRECATION_WARNINGS
 }
 
+void UContentBrowserAliasDataSource::RemoveUnusedCachedFilterData(const FContentBrowserDataFilterCacheIDOwner& IDOwner, TArrayView<const FName> InVirtualPathsInUse, const FContentBrowserDataFilter& DataFilter)
+{
+	FilterCache.RemoveUnusedCachedData(IDOwner, InVirtualPathsInUse, DataFilter);
+}
+
+void UContentBrowserAliasDataSource::ClearCachedFilterData(const FContentBrowserDataFilterCacheIDOwner& IDOwner)
+{
+	FilterCache.ClearCachedData(IDOwner);
+}
+
 FContentBrowserItemData UContentBrowserAliasDataSource::CreateAssetFolderItem(const FName InFolderPath)
 {
 	FName VirtualizedPath;
diff --git a/Engine/Plugins/Editor/ContentBrowser/ContentBrowserAliasDataSoure/Source/ContentBrowserAliasDataSource/Public/ContentBrowserAliasDataSource.h b/Engine/Plugins/Editor/ContentBrowser/ContentBrowserAliasDataSoure/Source/ContentBrowserAliasDataSource/Public/ContentBrowserAliasDataSource.h
--- a/Engine/Plugins/Editor/ContentBrowser/ContentBrowserAliasDataSoure/Source/ContentBrowserAliasDataSource/Public/ContentBrowserAliasDataSource.h
+++ b/Engine/Plugins/Editor/ContentBrowser/ContentBrowserAliasDataSoure/Source/ContentBrowserAliasDataSource/Public/ContentBrowserAliasDataSource.h
@@ -4,6 +4,7 @@
 
 
 #include "ContentBrowserAssetDataPayload.h"
+#include "ContentBrowserAssetDataSource.h"
 #include "ContentBrowserDataSource.h"
 
 #include "AssetRegistry/PathTree.h"
@@ -98,6 +99,10 @@ public:
 	virtual bool Legacy_TryConvertAssetDataToVirtualPath(const FAssetData& InAssetData, const bool InUseFolderPaths, FName& OutPath) override;
 	// ~ End UContentBrowserDataSource interface
 
+	virtual void RemoveUnusedCachedFilterData(const FContentBrowserDataFilterCacheIDOwner& IDOwner, TArrayView<const FName> InVirtualPathsInUse, const FContentBrowserDataFilter& DataFilter) override;
+
+	virtual void ClearCachedFilterData(const FContentBrowserDataFilterCacheIDOwner& IDOwner) override;
+
 	/** 
 	 * Add a list of aliases for a given asset.
 	 *
@@ -187,6 +192,8 @@ private:
 	TMap<FName, TArray<FContentBrowserUniqueAlias>> AliasesInPackagePath;
 	/** A set used for removing duplicate aliases in the same query, stored here to avoid constant reallocation */
 	TSet<FSoftObjectPath> AlreadyAddedOriginalAssets;
+
+	UContentBrowserAssetDataSource::FAssetDataSourceFilterCache FilterCache;
 };
 
 #if UE_ENABLE_INCLUDE_ORDER_DEPRECATED_IN_5_2
diff --git a/Engine/Plugins/Editor/ContentBrowser/ContentBrowserAssetDataSource/Source/ContentBrowserAssetDataSource/Private/ContentBrowserAssetDataSource.cpp b/Engine/Plugins/Editor/ContentBrowser/ContentBrowserAssetDataSource/Source/ContentBrowserAssetDataSource/Private/ContentBrowserAssetDataSource.cpp
--- a/Engine/Plugins/Editor/ContentBrowser/ContentBrowserAssetDataSource/Source/ContentBrowserAssetDataSource/Private/ContentBrowserAssetDataSource.cpp
+++ b/Engine/Plugins/Editor/ContentBrowser/ContentBrowserAssetDataSource/Source/ContentBrowserAssetDataSource/Private/ContentBrowserAssetDataSource.cpp
@@ -39,6 +39,9 @@
 
 #define LOCTEXT_NAMESPACE "ContentBrowserAssetDataSource"
 
+UContentBrowserAssetDataSource::FOnAssetDataSourcePathAdded UContentBrowserAssetDataSource::OnAssetPathAddedDelegate;
+UContentBrowserAssetDataSource::FOnAssetDataSourcePathRemoved UContentBrowserAssetDataSource::OnAssetPathRemovedDelegate;
+
 void UContentBrowserAssetDataSource::Initialize(const bool InAutoRegister)
 {
 	check(GIsEditor && !IsRunningCommandlet());
@@ -201,9 +204,10 @@
 	Super::Shutdown();
 }
 
-bool UContentBrowserAssetDataSource::PopulateAssetFilterInputParams(FAssetFilterInputParams& Params, UContentBrowserDataSource* DataSource, IAssetRegistry* InAssetRegistry, const FContentBrowserDataFilter& InFilter, FContentBrowserDataCompiledFilter& OutCompiledFilter, ICollectionManager* CollectionManager)
+bool UContentBrowserAssetDataSource::PopulateAssetFilterInputParams(FAssetFilterInputParams& Params, UContentBrowserDataSource* DataSource, IAssetRegistry* InAssetRegistry, const FContentBrowserDataFilter& InFilter, FContentBrowserDataCompiledFilter& OutCompiledFilter, ICollectionManager* CollectionManager, FAssetDataSourceFilterCache* InFilterCache)
 {
 	Params.CollectionManager = CollectionManager;
+	Params.AssetFilterCache = InFilterCache;
 
 	Params.bIncludeFolders = EnumHasAnyFlags(InFilter.ItemTypeFilter, EContentBrowserItemTypeFilter::IncludeFolders);
 	Params.bIncludeFiles = EnumHasAnyFlags(InFilter.ItemTypeFilter, EContentBrowserItemTypeFilter::IncludeFiles);
@@ -401,3 +405,647 @@
 	}
 
 	return true;
+}
+
+bool UContentBrowserAssetDataSource::CreateAssetFilter(FAssetFilterInputParams& Params, FName InPath, const FContentBrowserDataFilter& InFilter, FContentBrowserDataCompiledFilter& OutCompiledFilter, FSubPathEnumerationFunc* GetSubPackagePathsFunc)
+{
+	TRACE_CPUPROFILER_EVENT_SCOPE(UContentBrowserAssetDataSource::CreateAssetFilter);
+
+	// If we're not including files, then we can bail now as the rest of this function deals with assets
+	if (!Params.bIncludeFiles)
+	{
+		return false;
+	}
+
+	// If we are filtering all classes, then we can bail now as we won't return any content
+	if (Params.ClassPermissionList && Params.ClassPermissionList->IsDenyListAll() && !Params.UnsupportedClassFilter)
+	{
+		return false;
+	}
+
+	// If we are filtering out this path, then we can bail now as it won't return any content
+	if (Params.PathPermissionList && !InFilter.bRecursivePaths)
+	{
+		for (auto It = Params.InternalPaths.CreateIterator(); It; ++It)
+		{
+			if (!Params.PathPermissionList->PassesStartsWithFilter(*It))
+			{
+				It.RemoveCurrent();
+			}
+		}
+
+		if (Params.InternalPaths.Num() == 0)
+		{
+			return false;
+		}
+	}
+
+	auto DefaultEnumarePackagePaths = [&Params](FName InPath, TFunctionRef<bool(FName)> Callback, bool bIsRecursive)
+		{
+			Params.AssetRegistry->EnumerateSubPaths(InPath
+				, [&Callback](FName ChildPath)
+					{
+						return Callback(ChildPath);
+					}
+				, bIsRecursive);
+		};
+
+	FSubPathEnumerationFunc EnumeratePackagePaths = GetSubPackagePathsFunc ? *GetSubPackagePathsFunc : DefaultEnumarePackagePaths;
+
+	// Build inclusive asset filter
+	FARCompiledFilter CompiledInclusiveFilter;
+	{
+		// Build the basic inclusive filter from the given data
+		{
+			FARFilter InclusiveFilter;
+			if (Params.ObjectFilter)
+			{
+PRAGMA_DISABLE_DEPRECATION_WARNINGS
+				InclusiveFilter.ObjectPaths.Append(Params.ObjectFilter->ObjectNamesToInclude);
+PRAGMA_ENABLE_DEPRECATION_WARNINGS
+				InclusiveFilter.TagsAndValues.Append(Params.ObjectFilter->TagsAndValuesToInclude);
+				InclusiveFilter.bIncludeOnlyOnDiskAssets |= Params.ObjectFilter->bOnDiskObjectsOnly;
+			}
+			if (Params.PackageFilter)
+			{
+				InclusiveFilter.PackageNames.Append(Params.PackageFilter->PackageNamesToInclude);
+				InclusiveFilter.PackagePaths.Append(Params.PackageFilter->PackagePathsToInclude);
+				if (Params.PackageFilter->bRecursivePackagePathsToInclude)
+				{
+					for (const FName Path : Params.PackageFilter->PackagePathsToInclude)
+					{
+						EnumeratePackagePaths(Path, [&InclusiveFilter](FName ChildPath)
+							{
+								InclusiveFilter.PackagePaths.Add(ChildPath);
+								return true;
+							}
+							, Params.PackageFilter->bRecursivePackagePathsToInclude);
+					}
+				}
+			}
+			if (Params.ClassFilter)
+			{
+				InclusiveFilter.ClassPaths.Append(Params.ClassFilter->ClassNamesToInclude);
+				InclusiveFilter.bRecursiveClasses |= Params.ClassFilter->bRecursiveClassNamesToInclude;
+			}
+			if (Params.CollectionFilter)
+			{
+				TArray<FSoftObjectPath> ObjectPathsForCollections;
+				if (GetObjectPathsForCollections(Params.CollectionManager, Params.CollectionFilter->SelectedCollections, Params.CollectionFilter->bIncludeChildCollections, ObjectPathsForCollections) && ObjectPathsForCollections.Num() == 0)
+				{
+					// If we had collections but they contained no objects then we can bail as nothing will pass the filter
+					return false;
+				}
+				InclusiveFilter.SoftObjectPaths.Append(MoveTemp(ObjectPathsForCollections));
+			}
+
+#if DO_ENSURE
+			// Ensure paths do not have trailing slash	
+			static const FName RootPath = "/";
+
+			for (const FName ItPath : Params.InternalPaths)
+			{
+				ensure(ItPath == RootPath || !FStringView(FNameBuilder(ItPath)).EndsWith(TEXT('/')));
+			}
+
+			for (const FName ItPath : InclusiveFilter.PackagePaths)
+			{
+				ensure(ItPath == RootPath || !FStringView(FNameBuilder(ItPath)).EndsWith(TEXT('/')));
+			}
+#endif // DO_ENSURE
+
+			Params.AssetRegistry->CompileFilter(InclusiveFilter, CompiledInclusiveFilter);
+		}
+
+
+		// Add the backend class filtering to the unsupported asset filtering before the class permission are added
+		if (Params.ConvertToUnsupportedAssetDataFilter)
+		{
+			if (Params.UnsupportedClassFilter)
+			{
+				if (const FPathPermissionList* ClassPermissionList = Params.UnsupportedClassFilter->ClassPermissionList.Get())
+				{
+					if (ClassPermissionList->HasFiltering())
+					{
+						if (Params.ConvertToUnsupportedAssetDataFilter->InclusiveFilter.IsEmpty())
+						{
+							Params.ConvertToUnsupportedAssetDataFilter->InclusiveFilter.ClassPaths = CompiledInclusiveFilter.ClassPaths;
+						}
+						else
+						{
+							Params.ConvertToUnsupportedAssetDataFilter->InclusiveFilter.ClassPaths = Params.ConvertToUnsupportedAssetDataFilter->InclusiveFilter.ClassPaths.Intersect(CompiledInclusiveFilter.ClassPaths);
+						}
+					}
+				}
+			}
+		}
+
+
+		// Remove any inclusive paths that aren't under the set of internal paths that we want to enumerate
+		{
+			FARCompiledFilter CompiledInternalPathFilter;
+
+			if (!Params.AssetFilterCache || !Params.AssetFilterCache->GetCachedCompiledInternalPaths(InFilter, InPath, CompiledInternalPathFilter.PackagePaths))
+			{
+				/**
+				 * This filter is created by testing the paths while we are recursively going down the path hierarchy
+				 * This effective because it can stop exploring a sub part of the path three when the current path fail the attribute filter
+				 * Also after a certain depth it stop testing the paths against the attribute filter since the result will the same as the parent path
+				 */
+				TRACE_CPUPROFILER_EVENT_SCOPE(UContentBrowserAssetDataSource::CreateAssetFilter::CreateInternalPathFilter);
+			
+				CompiledInternalPathFilter.PackagePaths.Reserve(Params.InternalPaths.Num());
+
+				const int32 MaxDepthPathTestNeeded = ContentBrowserDataUtils::GetMaxFolderDepthRequiredForAttributeFilter();
+
+				// This builder is shared across calls because it is stack allocated and it could cause some issues in depth recursive calls
+				FNameBuilder PathBufferStr;
+				TFunction<void (FName, int32)> TestAndGatherChildPaths;
+				TestAndGatherChildPaths = [&InFilter, &CompiledInternalPathFilter, &TestAndGatherChildPaths, &EnumeratePackagePaths, &PathBufferStr, MaxDepthPathTestNeeded](FName ChildPath, int32 CurrentDepth)
+					{
+						PathBufferStr.Reset();
+						ChildPath.AppendString(PathBufferStr);
+						if (ContentBrowserDataUtils::PathPassesAttributeFilter(PathBufferStr, CurrentDepth, InFilter.ItemAttributeFilter))
+						{
+							CompiledInternalPathFilter.PackagePaths.Add(ChildPath);
+							++CurrentDepth;
+
+							if (CurrentDepth < MaxDepthPathTestNeeded)
+							{
+								constexpr bool bIsRecursive = false;
+								EnumeratePackagePaths(ChildPath
+									, [&TestAndGatherChildPaths, CurrentDepth](FName ChildPath)
+									{
+										TestAndGatherChildPaths(ChildPath, CurrentDepth);
+										return true;
+									}
+								, bIsRecursive);
+							}
+							else
+							{
+								constexpr bool bIsRecursive = true;
+								EnumeratePackagePaths(ChildPath
+									, [&CompiledInternalPathFilter](FName ChildPath)
+									{
+										CompiledInternalPathFilter.PackagePaths.Add(ChildPath);
+										return true;
+									}
+								, bIsRecursive);
+							}
+						}
+					};
+
+
+				for (const FName InternalPath : Params.InternalPaths)
+				{
+					PathBufferStr.Reset();
+					InternalPath.AppendString(PathBufferStr);
+					FStringView Path(PathBufferStr);
+					if (ContentBrowserDataUtils::PathPassesAttributeFilter(Path, 0, InFilter.ItemAttributeFilter))
+					{
+						CompiledInternalPathFilter.PackagePaths.Add(InternalPath);
+
+						if (InFilter.bRecursivePaths)
+						{
+							// Minus one because the test depth start at zero
+							const int32 CurrentDepth = ContentBrowserDataUtils::CalculateFolderDepthOfPath(Path) - 1;
+							if (CurrentDepth < MaxDepthPathTestNeeded)
+							{
+								constexpr bool bIsRecursive = false;
+								EnumeratePackagePaths(InternalPath
+									, [&TestAndGatherChildPaths, CurrentDepth](FName ChildPath)
+									{
+										TestAndGatherChildPaths(ChildPath, CurrentDepth);
+										return true;
+									}
+								, bIsRecursive);
+							}
+							else
+							{
+								constexpr bool bIsRecursive = true;
+								EnumeratePackagePaths(InternalPath
+									, [&CompiledInternalPathFilter](FName ChildPath)
+									{
+										CompiledInternalPathFilter.PackagePaths.Add(ChildPath);
+										return true;
+									}
+								, bIsRecursive);
+							}
+						}
+					}
+				}
+			
+				if (Params.AssetFilterCache)
+				{
+					Params.AssetFilterCache->CacheCompiledInternalPaths(InFilter, InPath, CompiledInternalPathFilter.PackagePaths);
+				}
+			}
+
+			if (CompiledInclusiveFilter.PackagePaths.Num() > 0)
+			{
+				// Explicit paths given - remove anything not in the internal paths set
+				// If the paths resolve as empty then the combined filter will return nothing and can be skipped
+				CompiledInclusiveFilter.PackagePaths = CompiledInclusiveFilter.PackagePaths.Intersect(CompiledInternalPathFilter.PackagePaths);
+				if (CompiledInclusiveFilter.PackagePaths.Num() == 0)
+				{
+					return false;
+				}
+			}
+			else
+			{
+				// No explicit paths given - just use the internal paths set
+				CompiledInclusiveFilter.PackagePaths = MoveTemp(CompiledInternalPathFilter.PackagePaths);
+			}
+		}
+
+		// Add the backend class filtering to the unsupported asset filtering before the class permission are added
+		if (Params.ConvertToUnsupportedAssetDataFilter)
+		{
+			if (Params.UnsupportedClassFilter)
+			{
+				if (const FPathPermissionList* ClassPermissionList = Params.UnsupportedClassFilter->ClassPermissionList.Get())
+				{
+					if (ClassPermissionList->HasFiltering())
+					{
+						Params.ConvertToUnsupportedAssetDataFilter->InclusiveFilter.ClassPaths.Append(CompiledInclusiveFilter.ClassPaths);
+					}
+				}
+			}
+		}
+
+		// Remove any inclusive paths that aren't in the explicit AllowList set
+		if (Params.PathPermissionList && Params.PathPermissionList->GetAllowList().Num() > 0)
+		{
+			FARCompiledFilter CompiledPathFilterAllowList;
+			{
+				CompiledPathFilterAllowList.PackagePaths.Reserve(Params.PathPermissionList->GetAllowList().Num());
+				for (const auto& AllowListPair : Params.PathPermissionList->GetAllowList())
+				{
+					FName PackageName = *AllowListPair.Key;
+					CompiledPathFilterAllowList.PackagePaths.Add(PackageName);
+
+					constexpr bool bIsRecursive = true;
+					EnumeratePackagePaths(PackageName, [&CompiledPathFilterAllowList](FName ChildPath)
+							{
+								CompiledPathFilterAllowList.PackagePaths.Add(ChildPath);
+								return true;
+							}
+							, bIsRecursive);
+
+				}
+			}
+
+			if (CompiledInclusiveFilter.PackagePaths.Num() > 0)
+			{
+				// Explicit paths given - remove anything not in the allow list paths set
+				// If the paths resolve as empty then the combined filter will return nothing and can be skipped
+				CompiledInclusiveFilter.PackagePaths = CompiledInclusiveFilter.PackagePaths.Intersect(CompiledPathFilterAllowList.PackagePaths);
+				if (CompiledInclusiveFilter.PackagePaths.Num() == 0)
+				{
+					return false;
+				}
+			}
+			else
+			{
+				// No explicit paths given - just use the allow list paths set
+				CompiledInclusiveFilter.PackagePaths = MoveTemp(CompiledPathFilterAllowList.PackagePaths);
+			}
+		}
+
+		// Remove any inclusive classes that aren't in the explicit allow list set
+		if (Params.ClassPermissionList && Params.ClassPermissionList->GetAllowList().Num() > 0)
+		{
+			FARCompiledFilter CompiledClassFilterAllowList;
+			{
+				FARFilter AllowListClassFilter;
+				for (const auto& AllowListPair : Params.ClassPermissionList->GetAllowList())
+				{
+					AllowListClassFilter.ClassPaths.Add(FTopLevelAssetPath(AllowListPair.Key));
+				}
+				Params.AssetRegistry->CompileFilter(AllowListClassFilter, CompiledClassFilterAllowList);
+			}
+
+			if (CompiledInclusiveFilter.ClassPaths.Num() > 0)
+			{
+				// Explicit classes given - remove anything not in the allow list class set
+				// If the classes resolve as empty then the combined filter will return nothing and can be skipped
+				CompiledInclusiveFilter.ClassPaths = CompiledInclusiveFilter.ClassPaths.Intersect(CompiledClassFilterAllowList.ClassPaths);
+				if (CompiledInclusiveFilter.ClassPaths.Num() == 0 && !Params.ConvertToUnsupportedAssetDataFilter)
+				{
+					return false;
+				}
+			}
+			else
+			{
+				// No explicit classes given - just use the allow list class set
+				CompiledInclusiveFilter.ClassPaths = MoveTemp(CompiledClassFilterAllowList.ClassPaths);
+			}
+		}
+	}
+
+	// Build exclusive asset filter
+	FARCompiledFilter CompiledExclusiveFilter;
+	{
+		// Build the basic exclusive filter from the given data
+		{
+			FARFilter ExclusiveFilter;
+			if (Params.ObjectFilter)
+			{
+PRAGMA_DISABLE_DEPRECATION_WARNINGS
+				ExclusiveFilter.ObjectPaths.Append(Params.ObjectFilter->ObjectNamesToExclude);
+PRAGMA_ENABLE_DEPRECATION_WARNINGS
+				ExclusiveFilter.TagsAndValues.Append(Params.ObjectFilter->TagsAndValuesToExclude);
+				ExclusiveFilter.bIncludeOnlyOnDiskAssets |= Params.ObjectFilter->bOnDiskObjectsOnly;
+			}
+			if (Params.PackageFilter)
+			{
+				ExclusiveFilter.PackageNames.Append(Params.PackageFilter->PackageNamesToExclude);
+				ExclusiveFilter.PackagePaths.Append(Params.PackageFilter->PackagePathsToExclude);
+				if (Params.PackageFilter->bRecursivePackagePathsToExclude)
+				{
+					for (const FName Path : Params.PackageFilter->PackagePathsToExclude)
+					{
+						EnumeratePackagePaths(Path, [&ExclusiveFilter](FName ChildPath)
+							{
+								ExclusiveFilter.PackagePaths.Add(ChildPath);
+								return true;
+							}
+							, Params.PackageFilter->bRecursivePackagePathsToExclude);
+					}
+				}
+			}
+			if (Params.ClassFilter)
+			{
+				ExclusiveFilter.ClassPaths.Append(Params.ClassFilter->ClassNamesToExclude);
+				ExclusiveFilter.bRecursiveClasses |= Params.ClassFilter->bRecursiveClassNamesToExclude;
+			}
+
+			Params.AssetRegistry->CompileFilter(ExclusiveFilter, CompiledExclusiveFilter);
+		}
+
+		// Add any exclusive paths that are in the explicit DenyList set
+		if (Params.PathPermissionList && Params.PathPermissionList->GetDenyList().Num() > 0)
+		{
+			CompiledExclusiveFilter.PackagePaths.Reserve(Params.PathPermissionList->GetDenyList().Num());
+			for (const auto& FilterPair : Params.PathPermissionList->GetDenyList())
+			{
+				FName Path = *FilterPair.Key;
+				CompiledExclusiveFilter.PackagePaths.Add(Path);
+
+				constexpr bool bIsRecursive = true;
+				EnumeratePackagePaths(Path, [&CompiledExclusiveFilter](FName ChildPath)
+					{
+						CompiledExclusiveFilter.PackagePaths.Add(ChildPath);
+						return true;
+					}
+					, bIsRecursive);
+			}
+		}
+
+		// Add any exclusive classes that are in the explicit DenyList set
+		if (Params.ClassPermissionList && Params.ClassPermissionList->GetDenyList().Num() > 0)
+		{
+			FARCompiledFilter CompiledClassFilter;
+			{
+				FARFilter ClassFilter;
+				for (const auto& FilterPair : Params.ClassPermissionList->GetDenyList())
+				{
+					ClassFilter.ClassPaths.Add(FTopLevelAssetPath(FilterPair.Key));
+				}
+				Params.AssetRegistry->CompileFilter(ClassFilter, CompiledClassFilter);
+			}
+
+			CompiledExclusiveFilter.ClassPaths.Append(CompiledClassFilter.ClassPaths);
+		}
+	}
+
+	// Apply our exclusive filter to the inclusive one to resolve cases where the exclusive filter cancels out the inclusive filter
+	// If any filter components resolve as empty then the combined filter will return nothing and can be skipped
+	{
+		if (CompiledInclusiveFilter.PackageNames.Num() > 0 && CompiledExclusiveFilter.PackageNames.Num() > 0)
+		{
+			CompiledInclusiveFilter.PackageNames = CompiledInclusiveFilter.PackageNames.Difference(CompiledExclusiveFilter.PackageNames);
+			if (CompiledInclusiveFilter.PackageNames.Num() == 0)
+			{
+				return false;
+			}
+			CompiledExclusiveFilter.PackageNames.Reset();
+		}
+		if (CompiledInclusiveFilter.PackagePaths.Num() > 0 && CompiledExclusiveFilter.PackagePaths.Num() > 0)
+		{
+			CompiledInclusiveFilter.PackagePaths = CompiledInclusiveFilter.PackagePaths.Difference(CompiledExclusiveFilter.PackagePaths);
+			if (CompiledInclusiveFilter.PackagePaths.Num() == 0)
+			{
+				return false;
+			}
+			CompiledExclusiveFilter.PackagePaths.Reset();
+		}
+		if (CompiledInclusiveFilter.SoftObjectPaths.Num() > 0 && CompiledExclusiveFilter.SoftObjectPaths.Num() > 0)
+		{
+			CompiledInclusiveFilter.SoftObjectPaths = CompiledInclusiveFilter.SoftObjectPaths.Difference(CompiledExclusiveFilter.SoftObjectPaths);
+			if (CompiledInclusiveFilter.SoftObjectPaths.Num() == 0)
+			{
+				return false;
+			}
+			CompiledExclusiveFilter.SoftObjectPaths.Reset();
+		}
+		if (CompiledInclusiveFilter.ClassPaths.Num() > 0 && CompiledExclusiveFilter.ClassPaths.Num() > 0)
+		{
+			CompiledInclusiveFilter.ClassPaths = CompiledInclusiveFilter.ClassPaths.Difference(CompiledExclusiveFilter.ClassPaths);
+			if (CompiledInclusiveFilter.ClassPaths.Num() == 0 && !!Params.ConvertToUnsupportedAssetDataFilter)
+			{
+				return false;
+			}
+			CompiledExclusiveFilter.ClassPaths.Reset();
+		}
+	}
+
+	// When InPath is a fully virtual folder such as /All, having no package paths is expected
+	if (CompiledInclusiveFilter.PackagePaths.Num() == 0)
+	{
+		// Leave bFilterExcludesAllAssets set to true
+		// Otherwise PackagePaths.Num() == 0 is interpreted as everything passes
+		return false;
+	}
+
+	// If we are enumerating recursively then the inclusive path list will already be fully filtered so just use that
+	if (Params.bIncludeFolders && InFilter.bRecursivePaths)
+	{
+		Params.AssetDataFilter->CachedSubPaths = CompiledInclusiveFilter.PackagePaths;
+		for (const FName InternalPath : Params.InternalPaths)
+		{
+			Params.AssetDataFilter->CachedSubPaths.Remove(InternalPath); // Remove the root as it's not a sub-path
+		}
+		Params.AssetDataFilter->CachedSubPaths.Sort(FNameLexicalLess()); // Sort as we enumerate these in parent->child order
+	}
+
+	// If we got this far then we have something in the filters and need to run the query
+	Params.AssetDataFilter->bFilterExcludesAllAssets = false;
+	Params.AssetDataFilter->InclusiveFilter = MoveTemp(CompiledInclusiveFilter);
+	Params.AssetDataFilter->ExclusiveFilter = MoveTemp(CompiledExclusiveFilter);
+
+
+	// Compile the filter to show the unsupported items
+	if (Params.ConvertToUnsupportedAssetDataFilter)
+	{
+		if (Params.UnsupportedClassFilter)
+		{
+			if (const FPathPermissionList* ClassPermissionList = Params.UnsupportedClassFilter->ClassPermissionList.Get())
+			{
+				if (ClassPermissionList->HasFiltering())
+				{
+					// Create a Backend filter for the unsupported items
+					{
+						// Cache the existing class path
+						TSet<FTopLevelAssetPath> InclusiveClassPath = MoveTemp(Params.ConvertToUnsupportedAssetDataFilter->InclusiveFilter.ClassPaths);
+						TSet<FTopLevelAssetPath> ExclusiveClassPath = MoveTemp(Params.ConvertToUnsupportedAssetDataFilter->ExclusiveFilter.ClassPaths);
+
+						// Remove temporally the class filtering from the asset data filter
+						TSet<FTopLevelAssetPath> AssetDataInclusiveClassPath = MoveTemp(Params.AssetDataFilter->InclusiveFilter.ClassPaths);
+						TSet<FTopLevelAssetPath> AssetDataExclusiveClassPath = MoveTemp(Params.AssetDataFilter->ExclusiveFilter.ClassPaths);
+
+						Params.ConvertToUnsupportedAssetDataFilter->InclusiveFilter = Params.AssetDataFilter->InclusiveFilter;
+						Params.ConvertToUnsupportedAssetDataFilter->ExclusiveFilter = Params.AssetDataFilter->ExclusiveFilter;
+
+						// Restore the class filtering
+						Params.AssetDataFilter->InclusiveFilter.ClassPaths = MoveTemp(AssetDataInclusiveClassPath);
+						Params.AssetDataFilter->ExclusiveFilter.ClassPaths = MoveTemp(AssetDataExclusiveClassPath);
+
+						Params.ConvertToUnsupportedAssetDataFilter->InclusiveFilter.ClassPaths = MoveTemp(InclusiveClassPath);
+						Params.ConvertToUnsupportedAssetDataFilter->ExclusiveFilter.ClassPaths = MoveTemp(ExclusiveClassPath);
+					}
+
+					FPathPermissionList* FolderPermissionList = Params.UnsupportedClassFilter->FolderPermissionList.Get();
+
+					// Compile the inclusive filter for where to show the unsupported asset
+					{
+						FARCompiledFilter CompiledShowInclusiveFilter;
+
+						// Only show the unsupported asset in the specified folders
+						if (FolderPermissionList && FolderPermissionList->HasFiltering())
+						{
+							FARFilter ShowInclusiveFilter;
+
+							ShowInclusiveFilter.bRecursivePaths = true;
+
+							ShowInclusiveFilter.PackagePaths.Reserve(FolderPermissionList->GetAllowList().Num());
+							for (const auto& AllowListPair : ClassPermissionList->GetAllowList())
+							{
+								ShowInclusiveFilter.PackagePaths.Add(*(AllowListPair.Key));
+							}
+
+							Params.AssetRegistry->CompileFilter(ShowInclusiveFilter, CompiledShowInclusiveFilter);
+						}
+
+						if (CompiledShowInclusiveFilter.IsEmpty())
+						{
+							if (Params.ConvertToUnsupportedAssetDataFilter->ShowInclusiveFilter.PackagePaths.IsEmpty())
+							{
+								Params.ConvertToUnsupportedAssetDataFilter->ShowInclusiveFilter.PackagePaths = Params.ConvertToUnsupportedAssetDataFilter->InclusiveFilter.PackagePaths;
+							}
+							else
+							{ 
+								Params.ConvertToUnsupportedAssetDataFilter->ShowInclusiveFilter.PackagePaths = Params.ConvertToUnsupportedAssetDataFilter->ShowInclusiveFilter.PackagePaths.Intersect(Params.ConvertToUnsupportedAssetDataFilter->InclusiveFilter.PackagePaths);
+							}
+						}
+						else
+						{
+							CompiledInclusiveFilter.PackagePaths = CompiledShowInclusiveFilter.PackagePaths.Intersect(Params.ConvertToUnsupportedAssetDataFilter->InclusiveFilter.PackagePaths);
+
+							if (Params.ConvertToUnsupportedAssetDataFilter->ShowInclusiveFilter.PackagePaths.IsEmpty())
+							{
+								Params.ConvertToUnsupportedAssetDataFilter->ShowInclusiveFilter.PackagePaths = MoveTemp(CompiledInclusiveFilter.PackagePaths);
+							}
+							else
+							{
+								Params.ConvertToUnsupportedAssetDataFilter->ShowInclusiveFilter.PackagePaths = Params.ConvertToUnsupportedAssetDataFilter->ShowInclusiveFilter.PackagePaths.Intersect(CompiledInclusiveFilter.PackagePaths);
+							}
+						}
+					}
+
+					// Compile the exclusive filter for where to show the unsupported asset
+					{
+						FARCompiledFilter CompiledShowExclusiveFilter;
+
+						// Only show the unsupported asset in the specified folders
+						if (FolderPermissionList && FolderPermissionList->HasFiltering())
+						{
+							FARFilter ShowExclusiveFilter;
+
+							ShowExclusiveFilter.bRecursivePaths = true;
+
+							ShowExclusiveFilter.PackagePaths.Reserve(FolderPermissionList->GetDenyList().Num());
+							for (const auto& DenyListPair : ClassPermissionList->GetDenyList())
+							{
+								ShowExclusiveFilter.PackagePaths.Add(*(DenyListPair.Key));
+							}
+
+							Params.AssetRegistry->CompileFilter(ShowExclusiveFilter, CompiledShowExclusiveFilter);
+						}
+
+						CompiledShowExclusiveFilter.PackagePaths.Append(Params.ConvertToUnsupportedAssetDataFilter->ExclusiveFilter.PackagePaths);
+
+						if (Params.ConvertToUnsupportedAssetDataFilter->ShowExclusiveFilter.PackagePaths.IsEmpty())
+						{
+							Params.ConvertToUnsupportedAssetDataFilter->ShowExclusiveFilter.PackagePaths = MoveTemp(CompiledShowExclusiveFilter.PackagePaths);
+						}
+						else
+						{
+							Params.ConvertToUnsupportedAssetDataFilter->ShowExclusiveFilter.PackagePaths.Append(MoveTemp(CompiledShowExclusiveFilter.PackagePaths));
+						}
+					}
+
+
+					// Compile the convert if fail inclusive filter
+					if (!ClassPermissionList->GetAllowList().IsEmpty())
+					{
+						FARCompiledFilter CompiledConvertIfFailInclusiveFilter;
+						FARFilter ConvertIfFailInclusiveFilter;
+
+						// Remove any inclusive classes that aren't in the explicit allow list set
+						ConvertIfFailInclusiveFilter.ClassPaths.Reserve(ClassPermissionList->GetAllowList().Num());
+						for (const auto& AllowListPair : ClassPermissionList->GetAllowList())
+						{
+							ConvertIfFailInclusiveFilter.ClassPaths.Add(FTopLevelAssetPath(AllowListPair.Key));
+						}
+
+						Params.AssetRegistry->CompileFilter(ConvertIfFailInclusiveFilter, CompiledConvertIfFailInclusiveFilter);
+					
+
+						if (Params.ConvertToUnsupportedAssetDataFilter->ConvertIfFailInclusiveFilter.ClassPaths.Num() > 0)
+						{
+							// Explicit classes given - remove anything not in the allow list class set
+							// If the classes resolve as empty then the combined filter will return nothing and can be skipped
+							Params.ConvertToUnsupportedAssetDataFilter->ConvertIfFailInclusiveFilter.ClassPaths = Params.ConvertToUnsupportedAssetDataFilter->ConvertIfFailInclusiveFilter.ClassPaths.Intersect(CompiledConvertIfFailInclusiveFilter.ClassPaths);
+						}
+						else
+						{
+							// No explicit classes given - just use the allow list class set
+							Params.ConvertToUnsupportedAssetDataFilter->ConvertIfFailInclusiveFilter.ClassPaths = MoveTemp(CompiledConvertIfFailInclusiveFilter.ClassPaths);
+						}
+					}
+
+					// Compile the convert if fail exclusive filter
+					if (!ClassPermissionList->GetDenyList().IsEmpty())
+					{
+						FARCompiledFilter CompiledConvertIfFailExclusiveFilter;
+						FARFilter ConvertIfFailExclusiveFilter;
+
+						// Add any exclusive classes that are in the explicit DenyList set
+						ConvertIfFailExclusiveFilter.ClassPaths.Reserve(ClassPermissionList->GetDenyList().Num());
+						for (const auto& FilterPair : ClassPermissionList->GetDenyList())
+						{
+							ConvertIfFailExclusiveFilter.ClassPaths.Add(FTopLevelAssetPath(FilterPair.Key));
+						}
+
+
+						Params.AssetRegistry->CompileFilter(ConvertIfFailExclusiveFilter, CompiledConvertIfFailExclusiveFilter);
+
+						Params.ConvertToUnsupportedAssetDataFilter->ConvertIfFailExclusiveFilter.ClassPaths.Append(MoveTemp(CompiledConvertIfFailExclusiveFilter.ClassPaths));
+					}
+				}
+			}
+		}
+	}
+
+	return true;
@@ -405,7 +1053,7 @@
 
 bool UContentBrowserAssetDataSource::CreateAssetFilter(FAssetFilterInputParams& Params, FName InPath, const FContentBrowserDataFilter& InFilter, FContentBrowserDataCompiledFilter& OutCompiledFilter, FCompileARFilterFunc CreateCompiledFilter)
 {
-	TRACE_CPUPROFILER_EVENT_SCOPE(UContentBrowserAssetDataSource::CreateAssetFilter);
+	TRACE_CPUPROFILER_EVENT_SCOPE(UContentBrowserAssetDataSource::LegacyCreateAssetFilter);
 
 	// If we're not including files, then we can bail now as the rest of this function deals with assets
 	if (!Params.bIncludeFiles)
@@ -518,90 +1166,25 @@
 		{
 			FARCompiledFilter CompiledInternalPathFilter;
 			{
-				/**
-				 * This filter is created by testing the paths while we are recursively going down the path hierarchy
-				 * This effective because it can stop exploring a sub part of the path three when the current path fail the attribute filter
-				 * Also after a certain depth it stop testing the paths against the attribute filter since the result will the same as the parent path
-				 */
-				TRACE_CPUPROFILER_EVENT_SCOPE(UContentBrowserAssetDataSource::CreateAssetFilter::CreateInternalPathFilter);
-			
-				CompiledInternalPathFilter.PackagePaths.Reserve(Params.InternalPaths.Num());
-
-				const int32 MaxDepthPathTestNeeded = ContentBrowserDataUtils::GetMaxFolderDepthRequiredForAttributeFilter();
-
-				// This builder is shared across calls because it is stack allocated and it could cause some issues in depth recursive calls
-				FNameBuilder PathBufferStr;
-				TFunction<void (FName, int32)> TestAndGatherChildPaths;
-				TestAndGatherChildPaths = [&InFilter, &CompiledInternalPathFilter, &TestAndGatherChildPaths, &Params, &PathBufferStr, MaxDepthPathTestNeeded](FName ChildPath, int32 CurrentDepth)
-					{
-						PathBufferStr.Reset();
-						ChildPath.AppendString(PathBufferStr);
-						if (ContentBrowserDataUtils::PathPassesAttributeFilter(PathBufferStr, CurrentDepth, InFilter.ItemAttributeFilter))
-						{
-							CompiledInternalPathFilter.PackagePaths.Add(ChildPath);
-							++CurrentDepth;
-
-							if (CurrentDepth < MaxDepthPathTestNeeded)
-							{
-								constexpr bool bIsRecursive = false;
-								Params.AssetRegistry->EnumerateSubPaths(ChildPath
-									, [&TestAndGatherChildPaths, CurrentDepth](FName ChildPath)
-									{
-										TestAndGatherChildPaths(ChildPath, CurrentDepth);
-										return true;
-									}
-								, bIsRecursive);
-							}
-							else
-							{
-								constexpr bool bIsRecursive = true;
-								Params.AssetRegistry->EnumerateSubPaths(ChildPath
-									, [&CompiledInternalPathFilter](FName ChildPath)
-									{
-										CompiledInternalPathFilter.PackagePaths.Add(ChildPath);
-										return true;
-									}
-								, bIsRecursive);
-							}
-						}
-					};
-
-
+				FARFilter InternalPathFilter;
 				for (const FName InternalPath : Params.InternalPaths)
 				{
-					PathBufferStr.Reset();
-					InternalPath.AppendString(PathBufferStr);
-					FStringView Path(PathBufferStr);
-					if (ContentBrowserDataUtils::PathPassesAttributeFilter(Path, 0, InFilter.ItemAttributeFilter))
+					InternalPathFilter.PackagePaths.Add(InternalPath);
+				}
+				InternalPathFilter.bRecursivePaths = InFilter.bRecursivePaths;
+				CreateCompiledFilter(InternalPathFilter, CompiledInternalPathFilter);
+
+				{
+					TRACE_CPUPROFILER_EVENT_SCOPE(UContentBrowserAssetDataSource::CreateAssetFilter::RemovePaths);
+			
+					// Remove paths that do not pass item attribute filter (Engine, Plugins, Developer, Localized, __ExternalActors__ etc..)
+					for (auto It = CompiledInternalPathFilter.PackagePaths.CreateIterator(); It; ++It)
 					{
-						CompiledInternalPathFilter.PackagePaths.Add(InternalPath);
-
-						if (InFilter.bRecursivePaths)
+						FNameBuilder PathStr(*It);
+						FStringView Path(PathStr);
+						if (!ContentBrowserDataUtils::PathPassesAttributeFilter(Path, 0, InFilter.ItemAttributeFilter))
 						{
-							// Minus one because the test depth start at zero
-							const int32 CurrentDepth = ContentBrowserDataUtils::CalculateFolderDepthOfPath(Path) - 1;
-							if (CurrentDepth < MaxDepthPathTestNeeded)
-							{
-								constexpr bool bIsRecursive = false;
-								Params.AssetRegistry->EnumerateSubPaths(InternalPath
-									, [&TestAndGatherChildPaths, CurrentDepth](FName ChildPath)
-									{
-										TestAndGatherChildPaths(ChildPath, CurrentDepth);
-										return true;
-									}
-								, bIsRecursive);
-							}
-							else
-							{
-								constexpr bool bIsRecursive = true;
-								Params.AssetRegistry->EnumerateSubPaths(InternalPath
-									, [&CompiledInternalPathFilter](FName ChildPath)
-									{
-										CompiledInternalPathFilter.PackagePaths.Add(ChildPath);
-										return true;
-									}
-								, bIsRecursive);
-							}
+							It.RemoveCurrent();
 						}
 					}
 				}
@@ -995,12 +1578,13 @@
 	return true;
 }
 
+
 void UContentBrowserAssetDataSource::CompileFilter(const FName InPath, const FContentBrowserDataFilter& InFilter, FContentBrowserDataCompiledFilter& OutCompiledFilter)
 {
 	TRACE_CPUPROFILER_EVENT_SCOPE(UContentBrowserAssetDataSource::CompileFilter);
 
 	FAssetFilterInputParams Params;
-	if (PopulateAssetFilterInputParams(Params, this, AssetRegistry, InFilter, OutCompiledFilter, CollectionManager))
+	if (PopulateAssetFilterInputParams(Params, this, AssetRegistry, InFilter, OutCompiledFilter, CollectionManager, &FilterCache))
 	{
 		const bool bCreatedPathFilter = CreatePathFilter(Params, InPath, InFilter, OutCompiledFilter, [this](FName Path, TFunctionRef<bool(FName)> Callback, bool bRecursive)
 		{
@@ -1016,10 +1600,7 @@
 				AssetRegistry->SetTemporaryCachingMode(bWasTemporaryCachingModeEnabled);
 			};
 
-			const bool bCreatedAssetFilter = CreateAssetFilter(Params, InPath, InFilter, OutCompiledFilter, [this](FARFilter& InputFilter, FARCompiledFilter& OutputFilter)
-			{
-				AssetRegistry->CompileFilter(InputFilter, OutputFilter);
-			});
+			const bool bCreatedAssetFilter = CreateAssetFilter(Params, InPath, InFilter, OutCompiledFilter);
 
 			if (bCreatedAssetFilter)
 			{
@@ -2075,6 +2656,16 @@
 PRAGMA_ENABLE_DEPRECATION_WARNINGS
 }
 
+void UContentBrowserAssetDataSource::RemoveUnusedCachedFilterData(const FContentBrowserDataFilterCacheIDOwner& IDOwner, TArrayView<const FName> InVirtualPathsInUse, const FContentBrowserDataFilter& DataFilter)
+{
+	FilterCache.RemoveUnusedCachedData(IDOwner, InVirtualPathsInUse, DataFilter);
+}
+
+void UContentBrowserAssetDataSource::ClearCachedFilterData(const FContentBrowserDataFilterCacheIDOwner& IDOwner)
+{
+	FilterCache.ClearCachedData(IDOwner);
+}
+
 bool UContentBrowserAssetDataSource::IsKnownContentPath(const FName InPackagePath) const
 {
 	FNameBuilder PackagePathStr(InPackagePath);
@@ -2272,6 +2863,18 @@
 	RecentlyPopulatedAssetFolders.Empty();
 	
 	QueueItemDataUpdate(FContentBrowserItemDataUpdate::MakeItemAddedUpdate(CreateAssetFolderItem(PathName)));
+
+	FStringView PathView(InPath);
+	// Minus one because the test depth start at zero
+	const int32 CurrentDepth = ContentBrowserDataUtils::CalculateFolderDepthOfPath(PathView) - 1;
+	int32 Index;
+	if (PathView.FindLastChar(TEXT('/'), Index))
+	{ 
+		uint32 PathNameHash = GetTypeHash(PathName);
+		FName ParentPath(PathView.Left(Index));
+		uint32 ParentPathHash = GetTypeHash(ParentPath);
+		OnAssetPathAddedDelegate.Broadcast(PathName, PathView, PathNameHash, ParentPath, ParentPathHash, CurrentDepth);
+	}
 }
 
 void UContentBrowserAssetDataSource::OnPathRemoved(const FString& InPath)
@@ -2282,6 +2885,8 @@
 	AssetFolderToAttributes.Remove(PathName);
 
 	QueueItemDataUpdate(FContentBrowserItemDataUpdate::MakeItemRemovedUpdate(CreateAssetFolderItem(PathName)));
+
+	OnAssetPathRemovedDelegate.Broadcast(PathName, GetTypeHash(PathName));
 }
 
 void UContentBrowserAssetDataSource::OnPathPopulated(const FAssetData& InAssetData)
@@ -2794,5 +3399,85 @@
 		&& ContentBrowserDataUtils::PathPassesAttributeFilter(Path, 0, InFilter.ItemAttributeFilter); // PassesAttributeFilter
 }
 
+UContentBrowserAssetDataSource::FAssetDataSourceFilterCache::FAssetDataSourceFilterCache()
+{
+	UContentBrowserAssetDataSource::OnAssetPathAddedDelegate.AddRaw(this, &UContentBrowserAssetDataSource::FAssetDataSourceFilterCache::OnPathAdded);
+	UContentBrowserAssetDataSource::OnAssetPathRemovedDelegate.AddRaw(this, &UContentBrowserAssetDataSource::FAssetDataSourceFilterCache::OnPathRemoved);
+}
+
+UContentBrowserAssetDataSource::FAssetDataSourceFilterCache::~FAssetDataSourceFilterCache()
+{
+	UContentBrowserAssetDataSource::OnAssetPathAddedDelegate.RemoveAll(this);
+	UContentBrowserAssetDataSource::OnAssetPathRemovedDelegate.RemoveAll(this);
+}
+
+bool UContentBrowserAssetDataSource::FAssetDataSourceFilterCache::GetCachedCompiledInternalPaths(const FContentBrowserDataFilter& InFilter, FName InVirtualPath, TSet<FName>& OutCompiledInternalPath) const
+{
+	// We only use the cache if the query if is recursive
+	if (InFilter.CacheID && InFilter.bRecursivePaths)
+	{
+		if (const FCachedDataPerID* CachedCompiledPathsForID = CachedCompiledInternalPaths.Find(InFilter.CacheID))
+		{
+			if (const TSet<FName>* CompiledPaths = CachedCompiledPathsForID->InternalPaths.Find(InVirtualPath))
+			{
+				OutCompiledInternalPath = *CompiledPaths;
+				return true;
+			}
+		}
+	}
+
+	return false;
+}
+
+void UContentBrowserAssetDataSource::FAssetDataSourceFilterCache::CacheCompiledInternalPaths(const FContentBrowserDataFilter& InFilter, FName InVirtualPath, const TSet<FName>& CompiledInternalPaths)
+{
+	// We only use the cache if the query if is recursive
+	if (InFilter.CacheID && InFilter.bRecursivePaths)
+	{
+		FCachedDataPerID& CachedCompiledPathsForID = CachedCompiledInternalPaths.FindOrAdd(InFilter.CacheID);
+		CachedCompiledPathsForID.InternalPaths.Add(InVirtualPath, CompiledInternalPaths);
+		CachedCompiledPathsForID.ItemAttributeFilter = InFilter.ItemAttributeFilter;
+	}
+}
+
+void UContentBrowserAssetDataSource::FAssetDataSourceFilterCache::RemoveUnusedCachedData(const FContentBrowserDataFilterCacheIDOwner& IDOwner, TArrayView<const FName> InVirtualPathsInUse, const FContentBrowserDataFilter& DataFilter)
+{
+	// we always clear the cache for now. This should be improved in some future changes
+	ClearCachedData(IDOwner);
+}
+
+void UContentBrowserAssetDataSource::FAssetDataSourceFilterCache::ClearCachedData(const FContentBrowserDataFilterCacheIDOwner& IDOwner)
+{
+	CachedCompiledInternalPaths.Remove(IDOwner);
+}
+
+void UContentBrowserAssetDataSource::FAssetDataSourceFilterCache::OnPathAdded(FName Path, FStringView PathString, uint32 PathHash, FName ParentPath, uint32 ParentPathHash, int32 PathDepth)
+{
+	for (TPair<FContentBrowserDataFilterCacheID, FCachedDataPerID>& CachedCompiledInternalPath : CachedCompiledInternalPaths)
+	{
+		if (ContentBrowserDataUtils::PathPassesAttributeFilter(PathString, PathDepth, CachedCompiledInternalPath.Value.ItemAttributeFilter))
+		{
+			for (TPair<FName, TSet<FName>>& CachedPaths : CachedCompiledInternalPath.Value.InternalPaths)
+			{
+				if (CachedPaths.Value.ContainsByHash(ParentPathHash, ParentPath))
+				{
+					CachedPaths.Value.AddByHash(PathHash, Path);
+				}
+			}
+		}
+	}
+}
+
+void UContentBrowserAssetDataSource::FAssetDataSourceFilterCache::OnPathRemoved(FName Path, uint32 PathHash)
+{
+	for (TPair<FContentBrowserDataFilterCacheID, FCachedDataPerID>& CachedCompiledInternalPath : CachedCompiledInternalPaths)
+	{
+		for (TPair<FName, TSet<FName>>& CachedPaths : CachedCompiledInternalPath.Value.InternalPaths)
+		{
+			CachedPaths.Value.RemoveByHash(PathHash, Path);
+		}
+	}
+}
+
 #undef LOCTEXT_NAMESPACE
 
diff --git a/Engine/Plugins/Editor/ContentBrowser/ContentBrowserAssetDataSource/Source/ContentBrowserAssetDataSource/Public/ContentBrowserAssetDataSource.h b/Engine/Plugins/Editor/ContentBrowser/ContentBrowserAssetDataSource/Source/ContentBrowserAssetDataSource/Public/ContentBrowserAssetDataSource.h
--- a/Engine/Plugins/Editor/ContentBrowser/ContentBrowserAssetDataSource/Source/ContentBrowserAssetDataSource/Public/ContentBrowserAssetDataSource.h
+++ b/Engine/Plugins/Editor/ContentBrowser/ContentBrowserAssetDataSource/Source/ContentBrowserAssetDataSource/Public/ContentBrowserAssetDataSource.h
@@ -25,6 +25,7 @@ class UToolMenu;
 class FAssetFolderContextMenu;
 class FAssetFileContextMenu;
 struct FCollectionNameType;
+class UContentBrowserAssetDataSource;
 class UContentBrowserToolbarMenuContext;
 
 USTRUCT()
@@ -115,6 +116,45 @@ public:
 
 	virtual void Shutdown() override;
 
+	/**
+	 * A struct used to cache some data to accelerate the compilation of the filters for the asset data source
+	 */
+	struct CONTENTBROWSERASSETDATASOURCE_API FAssetDataSourceFilterCache
+	{
+	public:
+		FAssetDataSourceFilterCache();
+		~FAssetDataSourceFilterCache();
+
+		FAssetDataSourceFilterCache(const FAssetDataSourceFilterCache&) = delete;
+		FAssetDataSourceFilterCache(FAssetDataSourceFilterCache&&) = delete;
+
+		FAssetDataSourceFilterCache& operator=(const FAssetDataSourceFilterCache&) = delete;
+		FAssetDataSourceFilterCache& operator=(FAssetDataSourceFilterCache&&) = delete;
+
+		void RemoveUnusedCachedData(const FContentBrowserDataFilterCacheIDOwner& IDOwner, TArrayView<const FName> InVirtualPathsInUse, const FContentBrowserDataFilter& DataFilter);
+		void ClearCachedData(const FContentBrowserDataFilterCacheIDOwner& IDOwner);
+
+	private:
+		friend UContentBrowserAssetDataSource;
+
+		bool GetCachedCompiledInternalPaths(const FContentBrowserDataFilter& InFilter, FName InVirtualPath, TSet<FName>& OutCompiledInternalPath) const;
+		void CacheCompiledInternalPaths(const FContentBrowserDataFilter& InFilter, FName InVirtualPath, const TSet<FName>& CompiledInternalPaths);
+
+		void OnPathAdded(FName Path, FStringView PathString, uint32 PathHash, FName ParentPath, uint32 ParentPathHash, int32 PathDepth);
+		void OnPathRemoved(FName Path, uint32 PathHash);
+
+		struct FCachedDataPerID
+		{
+			TMap<FName, TSet<FName>> InternalPaths;
+			EContentBrowserItemAttributeFilter ItemAttributeFilter = EContentBrowserItemAttributeFilter::IncludeAll;
+		};
+
+		/**
+		 * A cache for the recursive internal paths (Possible improvement: Rework the cache into a shared cache across multiple IDs for the simple filters and a not shared cache for the more complex filters)
+		 */
+		TMap<FContentBrowserDataFilterCacheID, FCachedDataPerID> CachedCompiledInternalPaths;
+	};
+
 	/**
 	 * All of the data necessary to generate a compiled filter for folders and assets
 	 */
@@ -124,6 +164,7 @@ public:
 
 		UContentBrowserDataSource* DataSource = nullptr;
 		ICollectionManager* CollectionManager = nullptr;
+		FAssetDataSourceFilterCache* AssetFilterCache = nullptr;
 		IAssetRegistry* AssetRegistry = nullptr;
 
 		const FContentBrowserDataObjectFilter* ObjectFilter = nullptr;
@@ -163,7 +204,8 @@ public:
 	 * 
 	 * @return false if it's not possible to display folders or assets, otherwise true
 	 */
-	static bool PopulateAssetFilterInputParams(FAssetFilterInputParams& Params, UContentBrowserDataSource* DataSource, IAssetRegistry* InAssetRegistry, const FContentBrowserDataFilter& InFilter, FContentBrowserDataCompiledFilter& OutCompiledFilter, ICollectionManager* CollectionManager = nullptr);
+	static bool PopulateAssetFilterInputParams(FAssetFilterInputParams& Params, UContentBrowserDataSource* DataSource, IAssetRegistry* InAssetRegistry, const FContentBrowserDataFilter& InFilter, FContentBrowserDataCompiledFilter& OutCompiledFilter, ICollectionManager* CollectionManager = nullptr, FAssetDataSourceFilterCache* InFilterCache = nullptr);
+	
 	/**
 	 * Call in CompileFilter() after PopulateAssetFilterInputParams() to fill OutCompiledFilter with an FContentBrowserCompiledAssetDataFilter capable of filtering folders
 	 *
@@ -176,6 +218,7 @@ public:
 	 * @return false if it's not possible to display any folders, otherwise true
 	 */
 	static bool CreatePathFilter(FAssetFilterInputParams& Params, FName InPath, const FContentBrowserDataFilter& InFilter, FContentBrowserDataCompiledFilter& OutCompiledFilter, FSubPathEnumerationFunc SubPathEnumeration);
+	
 	/**
 	 * Call in CompileFilter() after CreatePathFilter() to fill OutCompiledFilter with an FContentBrowserCompiledAssetDataFilter capable of filtering assets
 	 *
@@ -183,11 +226,15 @@ public:
 	 * @param InPath The input path supplied to CompileFilter()
 	 * @param InFilter The input filter supplied to CompileFilter()
 	 * @param OutCompiledFilter The output filer supplied to CompileFilter()
-	 * @param CreateCompiledFilter A function that generates an FARCompiledFilter from an input FARFilter
+	 * @param FGetSubPathsFunc A optional function to override how the package paths are expended.
 	 *
 	 * @return false if it's not possible to display any assets, otherwise true
 	 */
+	static bool CreateAssetFilter(FAssetFilterInputParams& Params, FName InPath, const FContentBrowserDataFilter& InFilter, FContentBrowserDataCompiledFilter& OutCompiledFilter, FSubPathEnumerationFunc* GetSubPackagePathsFunc = nullptr);
+
+	UE_DEPRECATED(5.3, "Use the override without the function to customize the compilation the ARFilter or duplicate this function logic. This function will no longer be supported and will be removed to allow future performence improvements")
 	static bool CreateAssetFilter(FAssetFilterInputParams& Params, FName InPath, const FContentBrowserDataFilter& InFilter, FContentBrowserDataCompiledFilter& OutCompiledFilter, FCompileARFilterFunc CreateCompiledFilter);
+	
 	/**
 	 * Call in EnumerateItemsMatchingFilter() to generate a list of folders that match the compiled filter.
 	 * It is the caller's responsibility to verify EContentBrowserItemTypeFilter::IncludeFolders is set before enumerating.
@@ -201,7 +248,8 @@ public:
 	 * @return 
 	 */
 	static void EnumerateFoldersMatchingFilter(UContentBrowserDataSource* DataSource, const FContentBrowserCompiledAssetDataFilter* AssetDataFilter, TFunctionRef<bool(FContentBrowserItemData&&)> InCallback, FSubPathEnumerationFunc SubPathEnumeration, FCreateFolderItemFunc CreateFolderItem);
-	/**
+	
+/**
 	 * Call in DoesItemPassFilter() to check if a folder passes the compiled asset data filter.
 	 * It is the caller's responsibility to verify EContentBrowserItemTypeFilter::IncludeFolders is set before enumerating.
 	 * 
@@ -325,6 +373,10 @@ public:
 
 	static bool PathPassesCompiledDataFilter(const FContentBrowserCompiledAssetDataFilter& InFilter, const FName InPath);
 
+	virtual void RemoveUnusedCachedFilterData(const FContentBrowserDataFilterCacheIDOwner& IDOwner, TArrayView<const FName> InVirtualPathsInUse, const FContentBrowserDataFilter& DataFilter) override;
+
+	virtual void ClearCachedFilterData(const FContentBrowserDataFilterCacheIDOwner& IDOwner) override;
+
 protected:
 	virtual void BuildRootPathVirtualTree() override;
 
@@ -442,6 +494,14 @@ private:
 	 * A cache of folders that have been populated since the last time any new asset folders were added.
 	 */
 	TMap<FName, EContentBrowserFolderAttributes> RecentlyPopulatedAssetFolders;
+
+	DECLARE_MULTICAST_DELEGATE_SixParams(FOnAssetDataSourcePathAdded, FName, FStringView, uint32, FName, uint32, int32);
+	static FOnAssetDataSourcePathAdded OnAssetPathAddedDelegate;
+
+	DECLARE_MULTICAST_DELEGATE_TwoParams(FOnAssetDataSourcePathRemoved, FName, uint32);
+	static FOnAssetDataSourcePathRemoved OnAssetPathRemovedDelegate;
+
+	FAssetDataSourceFilterCache FilterCache;
 };
 
 #if UE_ENABLE_INCLUDE_ORDER_DEPRECATED_IN_5_2
diff --git a/Engine/Source/Editor/ContentBrowser/Private/SAssetView.cpp b/Engine/Source/Editor/ContentBrowser/Private/SAssetView.cpp
--- a/Engine/Source/Editor/ContentBrowser/Private/SAssetView.cpp
+++ b/Engine/Source/Editor/ContentBrowser/Private/SAssetView.cpp
@@ -205,6 +205,7 @@ void SAssetView::Construct( const FArguments& InArgs )
 	ContentBrowserData->OnItemDataUpdated().AddSP(this, &SAssetView::HandleItemDataUpdated);
 	ContentBrowserData->OnItemDataRefreshed().AddSP(this, &SAssetView::RequestSlowFullListRefresh);
 	ContentBrowserData->OnItemDataDiscoveryComplete().AddSP(this, &SAssetView::HandleItemDataDiscoveryComplete);
+	FilterCacheID.Initialaze(ContentBrowserData);
 
 	FCollectionManagerModule& CollectionManagerModule = FCollectionManagerModule::GetModule();
 	CollectionManagerModule.Get().OnAssetsAddedToCollection().AddSP( this, &SAssetView::OnAssetsAddedToCollection );
@@ -1839,7 +1840,7 @@ bool SAssetView::IsValidSearchToken(const FString& Token) const
 	return true;
 }
 
-FContentBrowserDataFilter SAssetView::CreateBackendDataFilter() const
+FContentBrowserDataFilter SAssetView::CreateBackendDataFilter(bool bInvalidateCache) const
 {
 	// Assemble the filter using the current sources
 	// Force recursion when the user is searching
@@ -1904,6 +1905,23 @@ FContentBrowserDataFilter SAssetView::CreateBackendDataFilter() const
 		LegacyFilter.OnGetCustomSourceAssets = OnGetCustomSourceAssets;
 	}
 
+	DataFilter.CacheID = FilterCacheID;
+
+	if (bInvalidateCache)
+	{
+		if (SourcesData.IsIncludingVirtualPaths())
+		{
+			static const FName RootPath = "/";
+			const TArrayView<const FName> DataSourcePaths = SourcesData.HasVirtualPaths() ? MakeArrayView(SourcesData.VirtualPaths) : MakeArrayView(&RootPath, 1);
+			FilterCacheID.RemoveUnusedCachedData(DataSourcePaths, DataFilter);
+		}
+		else
+		{
+			// Not sure what is the right thing to do here so clear the cache
+			FilterCacheID.ClearCachedData();
+		}
+	}
+
 	return DataFilter;
 }
 
@@ -1958,7 +1976,8 @@ void SAssetView::RefreshSourceItems()
 
 		if (SourcesData.IsIncludingVirtualPaths() || SourcesData.HasCollections()) 
 		{
-			const FContentBrowserDataFilter DataFilter = CreateBackendDataFilter();
+			const bool bInvalidateFilterCache = true;
+			const FContentBrowserDataFilter DataFilter = CreateBackendDataFilter(bInvalidateFilterCache);
 
 			bWereItemsRecursivelyFiltered = DataFilter.bRecursivePaths;
 
@@ -2120,7 +2139,8 @@ void SAssetView::RefreshFilteredItems()
 		static const FName RootPath = "/";
 		const TArrayView<const FName> DataSourcePaths = SourcesData.HasVirtualPaths() ? MakeArrayView(SourcesData.VirtualPaths) : MakeArrayView(&RootPath, 1);
 
-		const FContentBrowserDataFilter DataFilter = CreateBackendDataFilter();
+		const bool bInvalidateFilterCache = false;
+		const FContentBrowserDataFilter DataFilter = CreateBackendDataFilter(bInvalidateFilterCache);
 
 		for (int32 FilterIdx = 0; FilterIdx < FrontendFilters->Num(); ++FilterIdx)
 		{
@@ -4723,7 +4743,8 @@ void SAssetView::HandleItemDataUpdated(TArrayView<const FContentBrowserItemDataU
 	TArray<FContentBrowserDataCompiledFilter> CompiledDataFilters;
 	if (SourcesData.IsIncludingVirtualPaths())
 	{
-		const FContentBrowserDataFilter DataFilter = CreateBackendDataFilter();
+		const bool bInvalidateFilterCache = false;
+		const FContentBrowserDataFilter DataFilter = CreateBackendDataFilter(bInvalidateFilterCache);
 
 		static const FName RootPath = "/";
 		const TArrayView<const FName> DataSourcePaths = SourcesData.HasVirtualPaths() ? MakeArrayView(SourcesData.VirtualPaths) : MakeArrayView(&RootPath, 1);
diff --git a/Engine/Source/Editor/ContentBrowser/Public/SAssetView.h b/Engine/Source/Editor/ContentBrowser/Public/SAssetView.h
--- a/Engine/Source/Editor/ContentBrowser/Public/SAssetView.h
+++ b/Engine/Source/Editor/ContentBrowser/Public/SAssetView.h
@@ -848,7 +848,7 @@ private:
 	/** Append the current effective backend filter (intersection of BackendFilter and SupportedFilter) to the given filter. */
 	void AppendBackendFilter(FARFilter& FilterToAppendTo) const;
 
-	FContentBrowserDataFilter CreateBackendDataFilter() const;
+	FContentBrowserDataFilter CreateBackendDataFilter(bool bInvalidateCache) const;
 
 	/** Handles updating the view when content items are changed */
 	void HandleItemDataUpdated(TArrayView<const FContentBrowserItemDataUpdate> InUpdatedItems);
@@ -1178,6 +1178,9 @@ private:
 	TArray<FString> HiddenColumnNames;
 
 	TArray<FAssetViewCustomColumn> CustomColumns;
+
+	/** An Id for the cache of the data sources for the filters compilation */
+	FContentBrowserDataFilterCacheIDOwner FilterCacheID;
 public:
 	bool ShouldColumnGenerateWidget(const FString ColumnName) const;
 };
diff --git a/Engine/Source/Editor/ContentBrowserData/Private/ContentBrowserDataFilter.cpp b/Engine/Source/Editor/ContentBrowserData/Private/ContentBrowserDataFilter.cpp
--- a/Engine/Source/Editor/ContentBrowserData/Private/ContentBrowserDataFilter.cpp
+++ b/Engine/Source/Editor/ContentBrowserData/Private/ContentBrowserDataFilter.cpp
@@ -2,6 +2,7 @@
 
 #include "ContentBrowserDataFilter.h"
 
+#include "ContentBrowserDataSubsystem.h"
 #include "Templates/Casts.h"
 
 FContentBrowserDataFilterList::FContentBrowserDataFilterList(const FContentBrowserDataFilterList& InOther)
@@ -83,3 +84,42 @@ void FContentBrowserDataFilterList::SetTo(const FContentBrowserDataFilterList& I
 		ScriptStructPtr->CopyScriptStruct(Filter.GetStructMemory(), OtherFilter.GetStructMemory());
 	}
 }
+
+void FContentBrowserDataFilterCacheIDOwner::Initialaze(UContentBrowserDataSubsystem* InContentBrowserDataSubsystem)
+{
+	if (DataSource != InContentBrowserDataSubsystem)
+	{
+		ClearCachedData();
+
+		if (InContentBrowserDataSubsystem)
+		{
+			UContentBrowserDataSubsystem::FContentBrowserFilterCacheApi::InitializeCacheIDOwner(*InContentBrowserDataSubsystem, *this);
+		}
+	}
+}
+
+void FContentBrowserDataFilterCacheIDOwner::RemoveUnusedCachedData(TArrayView<const FName> InVirtualPathsInUse, const FContentBrowserDataFilter& DataFilter) const
+{
+	if (const UContentBrowserDataSubsystem* ContentBrowserDataSubsystem = DataSource.Get())
+	{
+		UContentBrowserDataSubsystem::FContentBrowserFilterCacheApi::RemoveUnusedCachedData(*ContentBrowserDataSubsystem,
+			*this,
+			InVirtualPathsInUse,
+			DataFilter);
+	}
+}
+
+void FContentBrowserDataFilterCacheIDOwner::ClearCachedData() const
+{
+	if (const UContentBrowserDataSubsystem* ContentBrowserDataSubsystem = DataSource.Get())
+	{
+		UContentBrowserDataSubsystem::FContentBrowserFilterCacheApi::ClearCachedData(*ContentBrowserDataSubsystem, *this);
+	}
+}
+
+void FContentBrowserDataFilterCacheIDOwner::Reset()
+{
+	ClearCachedData();
+	ID = INDEX_NONE;
+	DataSource.Reset();
+}
diff --git a/Engine/Source/Editor/ContentBrowserData/Private/ContentBrowserDataSource.cpp b/Engine/Source/Editor/ContentBrowserData/Private/ContentBrowserDataSource.cpp
--- a/Engine/Source/Editor/ContentBrowserData/Private/ContentBrowserDataSource.cpp
+++ b/Engine/Source/Editor/ContentBrowserData/Private/ContentBrowserDataSource.cpp
@@ -153,6 +153,14 @@ bool UContentBrowserDataSource::TryConvertInternalPathToVirtual(const FName InIn
 	return true;
 }
 
+void UContentBrowserDataSource::RemoveUnusedCachedFilterData(const FContentBrowserDataFilterCacheIDOwner& IDOwner, TArrayView<const FName> InVirtualPathsInUse, const FContentBrowserDataFilter& DataFilter)
+{
+}
+
+void UContentBrowserDataSource::ClearCachedFilterData(const FContentBrowserDataFilterCacheIDOwner& IDOwner)
+{
+}
+
 void UContentBrowserDataSource::RootPathAdded(const FStringView InInternalPath)
 {
 	// Trim trailing slash
diff --git a/Engine/Source/Editor/ContentBrowserData/Private/ContentBrowserDataSubsystem.cpp b/Engine/Source/Editor/ContentBrowserData/Private/ContentBrowserDataSubsystem.cpp
--- a/Engine/Source/Editor/ContentBrowserData/Private/ContentBrowserDataSubsystem.cpp
+++ b/Engine/Source/Editor/ContentBrowserData/Private/ContentBrowserDataSubsystem.cpp
@@ -642,6 +642,51 @@ void UContentBrowserDataSubsystem::SetVirtualPathTreeNeedsRebuild()
 	}
 }
 
+
+void UContentBrowserDataSubsystem::FContentBrowserFilterCacheApi::InitializeCacheIDOwner(UContentBrowserDataSubsystem& Subsystem, FContentBrowserDataFilterCacheIDOwner& IDOwner)
+{
+	Subsystem.InitializeCacheIDOwner(IDOwner);
+}
+
+void UContentBrowserDataSubsystem::FContentBrowserFilterCacheApi::RemoveUnusedCachedData(const UContentBrowserDataSubsystem& Subsystem, const FContentBrowserDataFilterCacheIDOwner& IDOwner, TArrayView<const FName> InVirtualPathsInUse, const FContentBrowserDataFilter& DataFilter)
+{
+	Subsystem.RemoveUnusedCachedFilterData(IDOwner, InVirtualPathsInUse, DataFilter);
+}
+
+void UContentBrowserDataSubsystem::FContentBrowserFilterCacheApi::ClearCachedData(const UContentBrowserDataSubsystem& Subsystem, const FContentBrowserDataFilterCacheIDOwner& IDOwner)
+{
+	Subsystem.ClearCachedFilterData(IDOwner);
+}
+
+
+void UContentBrowserDataSubsystem::InitializeCacheIDOwner(FContentBrowserDataFilterCacheIDOwner& IDOwner)
+{
+	++LastCacheIDForFilter;
+	if (LastCacheIDForFilter == INDEX_NONE)
+	{
+		++LastCacheIDForFilter;
+	}
+
+	IDOwner.ID = LastCacheIDForFilter;
+	IDOwner.DataSource = this;
+}
+
+void UContentBrowserDataSubsystem::RemoveUnusedCachedFilterData(const FContentBrowserDataFilterCacheIDOwner& IDOwner, TArrayView<const FName> InVirtualPathsInUse, const FContentBrowserDataFilter& DataFilter) const 
+{
+	for (const TPair<FName, UContentBrowserDataSource*>& DataSource : AvailableDataSources)
+	{
+		DataSource.Value->RemoveUnusedCachedFilterData(IDOwner, InVirtualPathsInUse, DataFilter);
+	}
+}
+
+void UContentBrowserDataSubsystem::ClearCachedFilterData(const FContentBrowserDataFilterCacheIDOwner& IDOwner) const
+{
+	for (const TPair<FName, UContentBrowserDataSource*>& DataSource : AvailableDataSources)
+	{
+		DataSource.Value->ClearCachedFilterData(IDOwner);
+	}
+}
+
 void UContentBrowserDataSubsystem::HandleDataSourceRegistered(const FName& Type, IModularFeature* Feature)
 {
 	if (Type == UContentBrowserDataSource::GetModularFeatureTypeName())
diff --git a/Engine/Source/Editor/ContentBrowserData/Public/ContentBrowserDataFilter.h b/Engine/Source/Editor/ContentBrowserData/Public/ContentBrowserDataFilter.h
--- a/Engine/Source/Editor/ContentBrowserData/Public/ContentBrowserDataFilter.h
+++ b/Engine/Source/Editor/ContentBrowserData/Public/ContentBrowserDataFilter.h
@@ -23,6 +23,7 @@
 class FNamePermissionList;
 class FPathPermissionList;
 class UContentBrowserDataSource;
+class UContentBrowserDataSubsystem;
 
 /** Flags controlling which item types should be included */
 UENUM()
@@ -146,6 +147,30 @@ private:
 	TArray<FStructOnScope> TypedFilters;
 };
 
+struct CONTENTBROWSERDATA_API FContentBrowserDataFilterCacheID
+{
+public:
+	operator bool() const
+	{
+		return IsSet();
+	};
+
+	bool IsSet() const
+	{
+		return ID != INDEX_NONE;
+	}
+
+private:
+	friend struct FContentBrowserDataFilterCacheIDOwner;
+
+	int64 ID = INDEX_NONE;
+
+	friend uint32 GetTypeHash(const FContentBrowserDataFilterCacheID& CacheID)
+	{
+		return GetTypeHash(CacheID.ID);
+	}
+};
+
 /**
  * A filter used to control what is returned from Content Browser data queries.
  * @note The compiled version of this, FContentBrowserDataCompiledFilter, is produced via UContentBrowserDataSubsystem::CompileFilter.
@@ -174,6 +199,9 @@ public:
 
 	/** A list of extra filter structs to be interpreted by the Content Browser data sources */
 	FContentBrowserDataFilterList ExtraFilters;
+
+	/** An optional id used by the data sources to cache and reuse some data when compiling the filter(s) */
+	FContentBrowserDataFilterCacheID CacheID;
 };
 
 /**
@@ -325,3 +353,63 @@ public:
 
 	TSharedPtr<FPathPermissionList> FolderPermissionList;
 };
+
+/**
+ * ID used by the data sources to cache some data between the filter compilations
+ * How use the filter compilation cache.
+ * 1) Initialize the id by using the UContentBrowserDataSubsystem once.
+ * 2) When compiling the filters pass the Cache ID Owner to the cacheID of the ContentBrowserDataFilter.
+ * 3) When the filter settings change call RemoveUnusedCachedData to clean the cache and to remove the potential invalid data.
+ */
+struct CONTENTBROWSERDATA_API FContentBrowserDataFilterCacheIDOwner
+{
+public:
+	operator FContentBrowserDataFilterCacheID() const
+	{
+		FContentBrowserDataFilterCacheID CacheID;
+		CacheID.ID = ID;
+		return CacheID;
+	}
+
+	FContentBrowserDataFilterCacheIDOwner() = default;
+
+	FContentBrowserDataFilterCacheIDOwner(FContentBrowserDataFilterCacheIDOwner&& Other)
+		: ID(Other.ID)
+		, DataSource(MoveTemp(Other.DataSource))
+	{
+		Other.ID = INDEX_NONE;
+		Other.DataSource.Reset();
+	}
+
+	FContentBrowserDataFilterCacheIDOwner& operator=(FContentBrowserDataFilterCacheIDOwner&& Other)
+	{
+		ID = Other.ID;
+		Other.ID = INDEX_NONE;
+		Other.DataSource = MoveTemp(Other.DataSource);
+		Other.DataSource.Reset();
+
+		return *this;
+	}
+
+	~FContentBrowserDataFilterCacheIDOwner()
+	{
+		ClearCachedData();
+	}
+
+	FContentBrowserDataFilterCacheIDOwner(const FContentBrowserDataFilterCacheIDOwner&) = delete;
+	FContentBrowserDataFilterCacheIDOwner& operator=(const FContentBrowserDataFilterCacheIDOwner&) = delete;
+
+	void Initialaze(UContentBrowserDataSubsystem* InContentBrowserDataSubsystem);
+
+	void RemoveUnusedCachedData(TArrayView<const FName> InVirtualPathsInUse, const FContentBrowserDataFilter& DataFilter) const;
+
+	void ClearCachedData() const;
+
+	void Reset();
+
+private:
+	friend UContentBrowserDataSubsystem;
+
+	int64 ID = INDEX_NONE;
+	TWeakObjectPtr<UContentBrowserDataSubsystem> DataSource;
+};
diff --git a/Engine/Source/Editor/ContentBrowserData/Public/ContentBrowserDataSource.h b/Engine/Source/Editor/ContentBrowserData/Public/ContentBrowserDataSource.h
--- a/Engine/Source/Editor/ContentBrowserData/Public/ContentBrowserDataSource.h
+++ b/Engine/Source/Editor/ContentBrowserData/Public/ContentBrowserDataSource.h
@@ -834,6 +834,16 @@ PRAGMA_ENABLE_DEPRECATION_WARNINGS
 	 */
 	virtual bool TryConvertInternalPathToVirtual(const FName InInternalPath, FName& OutPath);
 
+	/**
+	 * Tell the data source to remove any cached data for the filter compilation that might not be needed any more.
+	 */
+	virtual void RemoveUnusedCachedFilterData(const FContentBrowserDataFilterCacheIDOwner& IDOwner, TArrayView<const FName> InVirtualPathsInUse, const FContentBrowserDataFilter& DataFilter);
+
+	/**
+	 * Tell the data source to remove the cached data for the filter compilation for this specific owner. 
+	 */
+	virtual void ClearCachedFilterData(const FContentBrowserDataFilterCacheIDOwner& IDOwner);
+
 protected:
 
 	/**
diff --git a/Engine/Source/Editor/ContentBrowserData/Public/ContentBrowserDataSubsystem.h b/Engine/Source/Editor/ContentBrowserData/Public/ContentBrowserDataSubsystem.h
--- a/Engine/Source/Editor/ContentBrowserData/Public/ContentBrowserDataSubsystem.h
+++ b/Engine/Source/Editor/ContentBrowserData/Public/ContentBrowserDataSubsystem.h
@@ -362,7 +362,29 @@ public:
 	 */
 	TSharedRef<FPathPermissionList>& GetEditableFolderPermissionList();
 
+
+	/**
+	 * Provide an partial access to private api for the filter cache ID owners
+	 * See FContentBrowserDataFilterCacheIDOwner declaration for how to use the caching for filter compilation
+	 */
+	struct FContentBrowserFilterCacheApi
+	{
+	private:
+		friend FContentBrowserDataFilterCacheIDOwner;
+
+		static void InitializeCacheIDOwner(UContentBrowserDataSubsystem& Subsystem, FContentBrowserDataFilterCacheIDOwner& IDOwner);
+		static void RemoveUnusedCachedData(const UContentBrowserDataSubsystem& Subsystem, const FContentBrowserDataFilterCacheIDOwner& IDOwner, TArrayView<const FName> InVirtualPathsInUse, const FContentBrowserDataFilter& DataFilter);
+		static void ClearCachedData(const UContentBrowserDataSubsystem& Subsystem, const FContentBrowserDataFilterCacheIDOwner& IDOwner);
+	};
+
 private:
+
+	void InitializeCacheIDOwner(FContentBrowserDataFilterCacheIDOwner& IDOwner);
+
+	void RemoveUnusedCachedFilterData(const FContentBrowserDataFilterCacheIDOwner& IDOwner, TArrayView<const FName> InVirtualPathsInUse, const FContentBrowserDataFilter& DataFilter) const;
+
+	void ClearCachedFilterData(const FContentBrowserDataFilterCacheIDOwner& IDOwner) const;
+
 	using FNameToDataSourceMap = TSortedMap<FName, UContentBrowserDataSource*, FDefaultAllocator, FNameFastLess>;
 
 	/**
@@ -505,6 +527,8 @@ private:
 
 	/** Permission list of folder paths we can edit */
 	TSharedRef<FPathPermissionList> EditableFolderPermissionList;
+
+	int64 LastCacheIDForFilter = INDEX_NONE;
 };
 
 /**
