Fixed an issue with the api of the asset definitions. The GetSourceFiles function didn't allow to resolve the source file paths in a way that is independant of the asset implementation details.

Added the capacitity to the AssetImportData to resolve a source file name without a loaded package.

This also fixes the issues where the editor is not able to find the source file during a reimport or simply navigating to that file via the content browser.

diff --git a/Engine/Plugins/Editor/ContentBrowser/ContentBrowserAssetDataSource/Source/ContentBrowserAssetDataSource/Private/AssetFileContextMenu.cpp b/Engine/Plugins/Editor/ContentBrowser/ContentBrowserAssetDataSource/Source/ContentBrowserAssetDataSource/Private/AssetFileContextMenu.cpp
--- a/Engine/Plugins/Editor/ContentBrowser/ContentBrowserAssetDataSource/Source/ContentBrowserAssetDataSource/Private/AssetFileContextMenu.cpp
+++ b/Engine/Plugins/Editor/ContentBrowser/ContentBrowserAssetDataSource/Source/ContentBrowserAssetDataSource/Private/AssetFileContextMenu.cpp
@@ -1223,27 +1223,40 @@ void FAssetFileContextMenu::ExecuteReimportWithNewFile(int32 SourceFileIndex /*=
 	if (const UAssetDefinition* AssetDefinition = UAssetDefinitionRegistry::Get()->GetAssetDefinitionForAsset(SelectedAssets[0]))
 	{
 		FAssetData SelectedAsset = SelectedAssets[0];
-		
-		AssetDefinition->GetSourceFiles(SelectedAsset, [&](const FAssetImportInfo& AssetImportInfo)
+
+		FAssetSourceFilesArgs GetSourceFilesArgs;
+		GetSourceFilesArgs.Assets = TConstArrayView<FAssetData>(&SelectedAsset, 1);
+
+		// Doesn't need to resolve the paths so it is a bit quicker
+		GetSourceFilesArgs.FilePathFormat = EPathUse::Display;
+
+		int32 SourceFileCount = 0;
+		AssetDefinition->GetSourceFiles(GetSourceFilesArgs, [&SourceFileCount, SourceFileIndex](const FAssetSourceFilesResult& AssetImportInfo)
 		{
-			int32 SourceFileIndexToReplace = SourceFileIndex;
-			//Check if the data is valid
-			if (SourceFileIndex == INDEX_NONE)
-			{
-				if (AssetImportInfo.SourceFiles.Num() > 1)
-				{
-					//Ask for a new file for the index 0
-					SourceFileIndexToReplace = 0;
-				}
-			}
-			else
+			++SourceFileCount;
+			if (SourceFileIndex < SourceFileCount)
 			{
-				check(AssetImportInfo.SourceFiles.IsValidIndex(SourceFileIndex));
+				return false;
 			}
 
-			TArray<UObject*> LoadedAssets = { SelectedAsset.GetAsset() };
-			FReimportManager::Instance()->ValidateAllSourceFileAndReimport(LoadedAssets, true, SourceFileIndexToReplace, true);
+			return true;
 		});
+
+		int32 SourceFileIndexToReplace = SourceFileIndex;
+		//Check if the data is valid
+		if (SourceFileIndex == INDEX_NONE)
+		{
+			//Ask for a new file for the index 0
+			SourceFileIndexToReplace = 0;
+		}
+		else
+		{
+			// Validate that the index is validate
+			check(SourceFileIndex >= 0 && SourceFileIndex < SourceFileCount);
+		}
+
+		TArray<UObject*> LoadedAssets = { SelectedAsset.GetAsset() };
+		FReimportManager::Instance()->ValidateAllSourceFileAndReimport(LoadedAssets, true, SourceFileIndexToReplace, true);
 	}
 }
 
@@ -1291,18 +1304,21 @@ void FAssetFileContextMenu::GetSelectedAssetSourceFilePaths(TArray<FString>& Out
 	{
 		if (const UAssetDefinition* AssetDefinition = UAssetDefinitionRegistry::Get()->GetAssetDefinitionForClass(AssetsByClassPair.Key))
 		{
-			for (const FAssetData& Asset : AssetsByClassPair.Value)
+			FAssetSourceFilesArgs GetSourceFilesArgs;
+			GetSourceFilesArgs.Assets = TConstArrayView<FAssetData>(AssetsByClassPair.Value);
+
+			AssetDefinition->GetSourceFiles(GetSourceFilesArgs, [&OutFilePaths](const FAssetSourceFilesResult& AssetImportInfo)
 			{
-				AssetDefinition->GetSourceFiles(Asset, [&](const FAssetImportInfo& AssetImportInfo)
-				{
-					OutValidSelectedAssetCount++;
-					for (const auto& SourceFile : AssetImportInfo.SourceFiles)
-					{
-						OutFilePaths.Add(SourceFile.RelativeFilename);
-						OutUniqueSourceFileLabels.AddUnique(SourceFile.DisplayLabelName);
-					}
-				});
-			}
+				OutFilePaths.Add(AssetImportInfo.FilePath);
+				return true;
+			});
+
+			GetSourceFilesArgs.FilePathFormat = EPathUse::Display;
+			AssetDefinition->GetSourceFiles(GetSourceFilesArgs, [&OutUniqueSourceFileLabels](const FAssetSourceFilesResult& AssetImportInfo)
+			{
+				OutUniqueSourceFileLabels.Add(AssetImportInfo.FilePath);
+				return true;
+			});
 		}
 	}
 }
diff --git a/Engine/Plugins/Editor/EngineAssetDefinitions/Source/Private/Table/AssetDefinition_CurveTable.cpp b/Engine/Plugins/Editor/EngineAssetDefinitions/Source/Private/Table/AssetDefinition_CurveTable.cpp
--- a/Engine/Plugins/Editor/EngineAssetDefinitions/Source/Private/Table/AssetDefinition_CurveTable.cpp
+++ b/Engine/Plugins/Editor/EngineAssetDefinitions/Source/Private/Table/AssetDefinition_CurveTable.cpp
@@ -157,15 +157,19 @@ namespace MenuExtension_CurveTable
 		const UContentBrowserAssetContextMenuContext* Context = UContentBrowserAssetContextMenuContext::FindContextWithAssets(InContext);
 		
 		TArray<FString> ImportPaths;
+
 		for (const FAssetData& Asset : Context->SelectedAssets)
 		{
-			UAssetDefinitionRegistry::Get()->GetAssetDefinitionForAsset(Asset)->GetSourceFiles(Asset, [&ImportPaths](const FAssetImportInfo& AssetImportInfo)
-			{
-				for (const auto& SourceFile : AssetImportInfo.SourceFiles)
+			FAssetSourceFilesArgs GetSourceFilesArgs;
+			GetSourceFilesArgs.Assets = TConstArrayView<FAssetData>(&Asset, 1);
+			if (const UAssetDefinition* AssetDefination = UAssetDefinitionRegistry::Get()->GetAssetDefinitionForAsset(Asset))
+			{ 
+				AssetDefination->GetSourceFiles(GetSourceFilesArgs, [&ImportPaths](const FAssetSourceFilesResult& AssetImportInfo)
 				{
-					ImportPaths.Add(SourceFile.RelativeFilename);
-				}
-			});
+					ImportPaths.Add(AssetImportInfo.FilePath);
+					return true;
+				});
+			}
 		}
 
 		TArray<FString> PotentialFileExtensions;
@@ -184,13 +188,16 @@ namespace MenuExtension_CurveTable
 		TArray<FString> ImportPaths;
 		for (const FAssetData& Asset : Context->SelectedAssets)
 		{
-			UAssetDefinitionRegistry::Get()->GetAssetDefinitionForAsset(Asset)->GetSourceFiles(Asset, [&ImportPaths](const FAssetImportInfo& AssetImportInfo)
+			FAssetSourceFilesArgs GetSourceFilesArgs;
+			GetSourceFilesArgs.Assets = TConstArrayView<FAssetData>(&Asset, 1);
+			if (const UAssetDefinition* AssetDefination = UAssetDefinitionRegistry::Get()->GetAssetDefinitionForAsset(Asset))
 			{
-				for (const auto& SourceFile : AssetImportInfo.SourceFiles)
+				AssetDefination->GetSourceFiles(GetSourceFilesArgs, [&ImportPaths](const FAssetSourceFilesResult& AssetImportInfo)
 				{
-					ImportPaths.Add(SourceFile.RelativeFilename);
-				}
-			});
+					ImportPaths.Add(AssetImportInfo.FilePath);
+					return true;
+				});
+			}
 		}
 
 		TArray<FString> PotentialFileExtensions;
diff --git a/Engine/Plugins/Editor/EngineAssetDefinitions/Source/Private/Table/AssetDefinition_DataTable.cpp b/Engine/Plugins/Editor/EngineAssetDefinitions/Source/Private/Table/AssetDefinition_DataTable.cpp
--- a/Engine/Plugins/Editor/EngineAssetDefinitions/Source/Private/Table/AssetDefinition_DataTable.cpp
+++ b/Engine/Plugins/Editor/EngineAssetDefinitions/Source/Private/Table/AssetDefinition_DataTable.cpp
@@ -228,13 +228,16 @@ namespace MenuExtension_DataTable
 		TArray<FString> ImportPaths;
 		for (const FAssetData& Asset : Context->SelectedAssets)
 		{
-			UAssetDefinitionRegistry::Get()->GetAssetDefinitionForAsset(Asset)->GetSourceFiles(Asset, [&ImportPaths](const FAssetImportInfo& AssetImportInfo)
-			{
-				for (const auto& SourceFile : AssetImportInfo.SourceFiles)
+			FAssetSourceFilesArgs GetSourceFilesArgs;
+			GetSourceFilesArgs.Assets = TConstArrayView<FAssetData>(&Asset, 1);
+			if (const UAssetDefinition* AssetDefination = UAssetDefinitionRegistry::Get()->GetAssetDefinitionForAsset(Asset))
+			{ 
+				AssetDefination->GetSourceFiles(GetSourceFilesArgs, [&ImportPaths](const FAssetSourceFilesResult& AssetImportInfo)
 				{
-					ImportPaths.Add(SourceFile.RelativeFilename);
-				}
-			});
+					ImportPaths.Add(AssetImportInfo.FilePath);
+					return true;
+				});
+			}
 		}
 
 		TArray<FString> PotentialFileExtensions;
@@ -253,13 +256,16 @@ namespace MenuExtension_DataTable
 		TArray<FString> ImportPaths;
 		for (const FAssetData& Asset : Context->SelectedAssets)
 		{
-			UAssetDefinitionRegistry::Get()->GetAssetDefinitionForAsset(Asset)->GetSourceFiles(Asset, [&ImportPaths](const FAssetImportInfo& AssetImportInfo)
-			{
-				for (const auto& SourceFile : AssetImportInfo.SourceFiles)
+			FAssetSourceFilesArgs GetSourceFilesArgs;
+			GetSourceFilesArgs.Assets = TConstArrayView<FAssetData>(&Asset, 1);
+			if (const UAssetDefinition* AssetDefination = UAssetDefinitionRegistry::Get()->GetAssetDefinitionForAsset(Asset))
+			{ 
+				AssetDefination->GetSourceFiles(GetSourceFilesArgs, [&ImportPaths](const FAssetSourceFilesResult& AssetImportInfo)
 				{
-					ImportPaths.Add(SourceFile.RelativeFilename);
-				}
-			});
+					ImportPaths.Add(AssetImportInfo.FilePath);
+					return true;
+				});
+			}
 		}
 
 		TArray<FString> PotentialFileExtensions;
diff --git a/Engine/Source/Editor/AssetDefinition/Private/AssetDefinition.cpp b/Engine/Source/Editor/AssetDefinition/Private/AssetDefinition.cpp
--- a/Engine/Source/Editor/AssetDefinition/Private/AssetDefinition.cpp
+++ b/Engine/Source/Editor/AssetDefinition/Private/AssetDefinition.cpp
@@ -76,6 +76,45 @@ TConstArrayView<FAssetCategoryPath> UAssetDefinition::GetAssetCategories() const
 	return Categories;
 }
 
+EAssetCommandResult UAssetDefinition::GetSourceFiles(const FAssetSourceFilesArgs& InArgs, TFunctionRef<bool(const FAssetSourceFilesResult& InSourceFile)> SourceFileFunc) const
+{
+	bool bFoundSomeData = false;
+
+	FString SourceFileTagData;
+	FAssetSourceFilesResult Result;
+	for (const FAssetData& Asset : InArgs.Assets)
+	{
+		if (Asset.GetTagValue(UObject::SourceFileTagName(), SourceFileTagData))
+		{
+			TOptional<FAssetImportInfo> ImportInfoOptional = FAssetImportInfo::FromJson(SourceFileTagData);
+			if (ImportInfoOptional.IsSet())
+			{
+				bFoundSomeData = true;
+				FAssetImportInfo& ImportInfo = ImportInfoOptional.GetValue();
+
+				for (FAssetImportInfo::FSourceFile& SourceFiles : ImportInfo.SourceFiles)
+				{
+					Result.FilePath = MoveTemp(SourceFiles.RelativeFilename);
+					Result.Timestamp = MoveTemp(SourceFiles.Timestamp);
+					Result.FileHash = MoveTemp(SourceFiles.FileHash);
+				
+					if (InArgs.FilePathFormat == EPathUse::AbsolutePath)
+					{
+						Result.FilePath = UAssetImportData::ResolveImportFilename(FStringView(Result.FilePath), Asset.PackageName.ToString());
+					}
+
+					if (!SourceFileFunc(Result))
+					{
+						return EAssetCommandResult::Handled;
+					}
+				}
+			}
+		}
+	}
+
+	return bFoundSomeData ? EAssetCommandResult::Handled : EAssetCommandResult::Unhandled;
+}
+
 EAssetCommandResult UAssetDefinition::GetSourceFiles(const FAssetData& InAsset, TFunctionRef<void(const FAssetImportInfo& AssetImportData)> SourceFileFunc) const
 {
 	FString SourceFileTagData;
diff --git a/Engine/Source/Editor/AssetDefinition/Public/AssetDefinition.h b/Engine/Source/Editor/AssetDefinition/Public/AssetDefinition.h
--- a/Engine/Source/Editor/AssetDefinition/Public/AssetDefinition.h
+++ b/Engine/Source/Editor/AssetDefinition/Public/AssetDefinition.h
@@ -47,6 +47,21 @@ enum class EAssetMergeResult : uint8
 	Cancelled,
 };
 
+enum class EPathUse : uint8
+{
+	/**
+	 * Present the file in a UI friendly format. 
+	 * Note that the path might not always be relative.
+	 * It cannot be used for other purposes then presenting it to the user since the details of how to resolve that path might change based on the asset implementation.
+	 */
+	Display,
+
+	/**
+	 * Resolve the path into an absolute path.
+	 */
+	AbsolutePath
+};
+
 struct FAssetArgs
 {
 	FAssetArgs() { }
@@ -116,10 +131,21 @@ struct FAssetActivateArgs : public FAssetArgs
 	EAssetActivationMethod ActivationMethod;
 };
 
-struct FAssetSourceFile
+struct FAssetSourceFilesArgs : public FAssetArgs
 {
-	FString DisplayLabelName;
-	FString RelativeFilename;
+	EPathUse FilePathFormat = EPathUse::AbsolutePath;
+};
+
+struct FAssetSourceFilesResult
+{
+	/** The file path in the format requested. */
+	FString FilePath;
+
+	/** The timestamp of the file when it was imported (as UTC). 0 when unknown. */
+	FDateTime Timestamp;
+
+	/** The MD5 hash of the file when it was imported. Invalid when unknown. */
+	FMD5Hash FileHash;
 };
 
 struct FAssetMergeResults
@@ -493,6 +519,15 @@ public:
 	}
 
 	// Source Files
+
+	/**
+	 * Return the source files that was used to generate/import the asset
+	 * @param InArgs The asset data of the assets we want the source files from and in which format we want the file path to be.
+	 * @param SourceFileFunc A function that is called for each source file found. The call back must return true to continue the enumeration.
+	 */
+	virtual EAssetCommandResult GetSourceFiles(const FAssetSourceFilesArgs& InArgs, TFunctionRef<bool(const FAssetSourceFilesResult& InSourceFile)> SourceFileFunc) const;
+
+	UE_DEPRECATED(5.3, "This override will be removed because it doesn't account that the resolution of a relative path to its absolute path may varies based on the asset implementation.")
 	virtual EAssetCommandResult GetSourceFiles(const FAssetData& InAsset, TFunctionRef<void(const FAssetImportInfo& AssetImportData)> SourceFileFunc) const;
 
 	// Diffing Assets
diff --git a/Engine/Source/Editor/UnrealEd/Private/Editor.cpp b/Engine/Source/Editor/UnrealEd/Private/Editor.cpp
--- a/Engine/Source/Editor/UnrealEd/Private/Editor.cpp
+++ b/Engine/Source/Editor/UnrealEd/Private/Editor.cpp
@@ -861,17 +861,29 @@ void FReimportManager::GetNewReimportPath(UObject* Obj, TArray<FString>& InOutFi
 			FAssetData AssetData(Obj);
 			const UAssetDefinition* AssetDefinition = UAssetDefinitionRegistry::Get()->GetAssetDefinitionForAsset(AssetData);
 
-			TArray<FAssetImportInfo::FSourceFile> OutSourceAssets;
-			AssetDefinition->GetSourceFiles(AssetData, [&OutSourceAssets](const FAssetImportInfo& ImportInfo)
+			FAssetSourceFilesArgs GetSourceFilesArgs;
+			GetSourceFilesArgs.Assets = TConstArrayView<FAssetData>(&AssetData, 1);
+			GetSourceFilesArgs.FilePathFormat = EPathUse::Display;
+
+			int32 SourceFileCount = 0;
+			FString SourceDisplayLabel;
+			AssetDefinition->GetSourceFiles(GetSourceFilesArgs, [&SourceFileCount, &SourceDisplayLabel, SourceFileIndex](const FAssetSourceFilesResult& AssetImportInfo)
 			{
-				OutSourceAssets.Append(ImportInfo.SourceFiles);
+				++SourceFileCount;
+				if (SourceFileIndex < SourceFileCount)
+				{
+					SourceDisplayLabel = AssetImportInfo.FilePath;
+					return false;
+				}
+
+				return true;
 			});
 			
-			if (OutSourceAssets.IsValidIndex(SourceFileIndex))
+			if (SourceFileIndex >= 0 && SourceFileIndex < SourceFileCount)
 			{
 				Title = FString::Printf(TEXT("%s %s %s: %s"),
 					*NSLOCTEXT("ReimportManager", "ImportDialogTitleLabelPart1", "Select").ToString(),
-					*OutSourceAssets[SourceFileIndex].DisplayLabelName,
+					*SourceDisplayLabel,
 					*NSLOCTEXT("ReimportManager", "ImportDialogTitleLabelPart2", "Source File For").ToString(),
 					*Obj->GetName());
 			}
diff --git a/Engine/Source/Runtime/Engine/Classes/EditorFramework/AssetImportData.h b/Engine/Source/Runtime/Engine/Classes/EditorFramework/AssetImportData.h
--- a/Engine/Source/Runtime/Engine/Classes/EditorFramework/AssetImportData.h
+++ b/Engine/Source/Runtime/Engine/Classes/EditorFramework/AssetImportData.h
@@ -155,6 +155,7 @@ public:
 
 	/** Resolve a filename that is relative to either the specified package, BaseDir() or absolute */
 	static ENGINE_API FString ResolveImportFilename(const FString& InRelativePath, const UPackage* Outermost);
+	static ENGINE_API FString ResolveImportFilename(FStringView InRelativePath, FStringView OutermostPath);
 
 	/** Convert an absolute import path so that it's relative to either this object's package, BaseDir() or leave it absolute */
 	static ENGINE_API FString SanitizeImportFilename(const FString& InPath, const UPackage* Outermost);
diff --git a/Engine/Source/Runtime/Engine/Private/EditorFramework/AssetImportData.cpp b/Engine/Source/Runtime/Engine/Private/EditorFramework/AssetImportData.cpp
--- a/Engine/Source/Runtime/Engine/Private/EditorFramework/AssetImportData.cpp
+++ b/Engine/Source/Runtime/Engine/Private/EditorFramework/AssetImportData.cpp
@@ -3,6 +3,7 @@
 #include "EditorFramework/AssetImportData.h"
 #include "HAL/FileManager.h"
 #include "Misc/Paths.h"
+#include "Misc/PathViews.h"
 #include "Misc/PackageName.h"
 #include "Serialization/JsonSerializer.h"
 #include "UObject/Package.h"
@@ -351,17 +352,31 @@ FString UAssetImportData::SanitizeImportFilename(const FString& InPath, const FS
 FString UAssetImportData::ResolveImportFilename(const FString& InRelativePath, const UPackage* Outermost)
 {
 	if (Outermost)
+	{ 
+		return ResolveImportFilename(InRelativePath, FStringView(Outermost->GetPathName()));
+	}
+
+	return ResolveImportFilename(InRelativePath, FStringView());
+}
+
+FString UAssetImportData::ResolveImportFilename(FStringView InRelativePath, FStringView OutermostPath)
+{
+	if (!OutermostPath.IsEmpty())
 	{
 		// Relative to the package filename?
-		const FString PathRelativeToPackage = FPaths::GetPath(FPackageName::LongPackageNameToFilename(Outermost->GetPathName())) / InRelativePath;
-		FString FullConvertPath = FPaths::ConvertRelativePathToFull(PathRelativeToPackage);
-		if (FPaths::FileExists(FullConvertPath))
+		FStringBuilderBase Path;
+		FPathViews::Append(Path, FPathViews::GetPath(FPackageName::LongPackageNameToFilename(FString(OutermostPath))), InRelativePath);
+
+		
+		FPathViews::ToAbsolutePathInline(Path);
+		FString ResolvedFile = Path.ToString();
+		if (FPaths::FileExists(ResolvedFile))
 		{
 			//FileExist return true when testing Path like c:/../folder1/filename. ConvertRelativePathToFull specify having .. in front of a drive letter is an error.
 			//It is relative to package only if the conversion to full path is successful.
-			if (FullConvertPath.Find(TEXT("..")) == INDEX_NONE)
+			if (ResolvedFile.Find(TEXT("..")) == INDEX_NONE)
 			{
-				return FullConvertPath;
+				return ResolvedFile;
 			}
 		}
 	}
@@ -388,7 +403,9 @@ FString UAssetImportData::ResolveImportFilename(const FString& InRelativePath, c
 #endif
 
 	// Convert relative paths
-	return FPaths::ConvertRelativePathToFull(InRelativePath);
+	FStringBuilderBase Path;
+	FPathViews::ToAbsolutePath(InRelativePath, Path);
+	return Path.ToString();
 }
 
 FString UAssetImportData::ResolveImportFilename(const FString& InRelativePath) const
