Adding an missing count to the number of assets we were able to retrive some information about their source file in the  FAssetFileContextMenu::GetSelectedAssetSourceFilePaths function.

diff --git a/Engine/Plugins/Editor/ContentBrowser/ContentBrowserAssetDataSource/Source/ContentBrowserAssetDataSource/Private/AssetFileContextMenu.cpp b/Engine/Plugins/Editor/ContentBrowser/ContentBrowserAssetDataSource/Source/ContentBrowserAssetDataSource/Private/AssetFileContextMenu.cpp
--- a/Engine/Plugins/Editor/ContentBrowser/ContentBrowserAssetDataSource/Source/ContentBrowserAssetDataSource/Private/AssetFileContextMenu.cpp
+++ b/Engine/Plugins/Editor/ContentBrowser/ContentBrowserAssetDataSource/Source/ContentBrowserAssetDataSource/Private/AssetFileContextMenu.cpp
@@ -1307,7 +1307,7 @@ void FAssetFileContextMenu::GetSelectedAssetSourceFilePaths(TArray<FString>& Out
 			FAssetSourceFilesArgs GetSourceFilesArgs;
 			GetSourceFilesArgs.Assets = TConstArrayView<FAssetData>(AssetsByClassPair.Value);
 
-			AssetDefinition->GetSourceFiles(GetSourceFilesArgs, [&OutFilePaths](const FAssetSourceFilesResult& AssetImportInfo)
+			EAssetCommandResult Result = AssetDefinition->GetSourceFiles(GetSourceFilesArgs, [&OutFilePaths](const FAssetSourceFilesResult& AssetImportInfo)
 			{
 				OutFilePaths.Add(AssetImportInfo.FilePath);
 				return true;
@@ -1319,6 +1319,12 @@ void FAssetFileContextMenu::GetSelectedAssetSourceFilePaths(TArray<FString>& Out
 				OutUniqueSourceFileLabels.Add(AssetImportInfo.FilePath);
 				return true;
 			});
+
+			if (Result == EAssetCommandResult::Handled)
+			{
+				// We found some import data for the asset if the call is handled
+				++OutValidSelectedAssetCount;
+			}
 		}
 	}
 }
