Fixed an issue where the contextual menu for the asset reimport would consider multiple assets to be only one asset when all the asset are of the same type.

diff --git a/Engine/Plugins/Editor/ContentBrowser/ContentBrowserAssetDataSource/Source/ContentBrowserAssetDataSource/Private/AssetFileContextMenu.cpp b/Engine/Plugins/Editor/ContentBrowser/ContentBrowserAssetDataSource/Source/ContentBrowserAssetDataSource/Private/AssetFileContextMenu.cpp
--- a/Engine/Plugins/Editor/ContentBrowser/ContentBrowserAssetDataSource/Source/ContentBrowserAssetDataSource/Private/AssetFileContextMenu.cpp
+++ b/Engine/Plugins/Editor/ContentBrowser/ContentBrowserAssetDataSource/Source/ContentBrowserAssetDataSource/Private/AssetFileContextMenu.cpp
@@ -1299,31 +1299,35 @@ void FAssetFileContextMenu::GetSelectedAssetSourceFilePaths(TArray<FString>& Out
 	GetSelectedAssetsByClass(SelectedAssetsByClass);
 
 	OutValidSelectedAssetCount = 0;
+
+	FAssetSourceFilesArgs GetSourceFilesArgs;
 	// Get the source file paths for the assets of each type
 	for (const auto& AssetsByClassPair : SelectedAssetsByClass)
 	{
 		if (const UAssetDefinition* AssetDefinition = UAssetDefinitionRegistry::Get()->GetAssetDefinitionForClass(AssetsByClassPair.Key))
 		{
-			FAssetSourceFilesArgs GetSourceFilesArgs;
-			GetSourceFilesArgs.Assets = TConstArrayView<FAssetData>(AssetsByClassPair.Value);
-
-			EAssetCommandResult Result = AssetDefinition->GetSourceFiles(GetSourceFilesArgs, [&OutFilePaths](const FAssetSourceFilesResult& AssetImportInfo)
+			for (const FAssetData& AssetData :  AssetsByClassPair.Value)
 			{
-				OutFilePaths.Add(AssetImportInfo.FilePath);
-				return true;
-			});
+				GetSourceFilesArgs.Assets = TConstArrayView<FAssetData>(&AssetData, 1);
 
-			GetSourceFilesArgs.FilePathFormat = EPathUse::Display;
-			AssetDefinition->GetSourceFiles(GetSourceFilesArgs, [&OutUniqueSourceFileLabels](const FAssetSourceFilesResult& AssetImportInfo)
-			{
-				OutUniqueSourceFileLabels.Add(AssetImportInfo.FilePath);
-				return true;
-			});
+				EAssetCommandResult Result = AssetDefinition->GetSourceFiles(GetSourceFilesArgs, [&OutFilePaths](const FAssetSourceFilesResult& AssetImportInfo)
+				{
+					OutFilePaths.Add(AssetImportInfo.FilePath);
+					return true;
+				});
 
-			if (Result == EAssetCommandResult::Handled)
-			{
-				// We found some import data for the asset if the call is handled
-				++OutValidSelectedAssetCount;
+				GetSourceFilesArgs.FilePathFormat = EPathUse::Display;
+				AssetDefinition->GetSourceFiles(GetSourceFilesArgs, [&OutUniqueSourceFileLabels](const FAssetSourceFilesResult& AssetImportInfo)
+				{
+					OutUniqueSourceFileLabels.Add(AssetImportInfo.FilePath);
+					return true;
+				});
+
+				if (Result == EAssetCommandResult::Handled)
+				{
+					// We found some import data for the asset if the call is handled
+					++OutValidSelectedAssetCount;
+				}
 			}
 		}
 	}
