Fixed an issue where the contextual menu interpreted the relative file path of a source file an the display label instead of using the actual display label store with the source file.

diff --git a/Engine/Plugins/Editor/ContentBrowser/ContentBrowserAssetDataSource/Source/ContentBrowserAssetDataSource/Private/AssetFileContextMenu.cpp b/Engine/Plugins/Editor/ContentBrowser/ContentBrowserAssetDataSource/Source/ContentBrowserAssetDataSource/Private/AssetFileContextMenu.cpp
--- a/Engine/Plugins/Editor/ContentBrowser/ContentBrowserAssetDataSource/Source/ContentBrowserAssetDataSource/Private/AssetFileContextMenu.cpp
+++ b/Engine/Plugins/Editor/ContentBrowser/ContentBrowserAssetDataSource/Source/ContentBrowserAssetDataSource/Private/AssetFileContextMenu.cpp
@@ -1310,16 +1310,10 @@ void FAssetFileContextMenu::GetSelectedAssetSourceFilePaths(TArray<FString>& Out
 			{
 				GetSourceFilesArgs.Assets = TConstArrayView<FAssetData>(&AssetData, 1);
 
-				EAssetCommandResult Result = AssetDefinition->GetSourceFiles(GetSourceFilesArgs, [&OutFilePaths](const FAssetSourceFilesResult& AssetImportInfo)
+				EAssetCommandResult Result = AssetDefinition->GetSourceFiles(GetSourceFilesArgs, [&OutFilePaths, &OutUniqueSourceFileLabels](const FAssetSourceFilesResult& AssetImportInfo)
 				{
 					OutFilePaths.Add(AssetImportInfo.FilePath);
-					return true;
-				});
-
-				GetSourceFilesArgs.FilePathFormat = EPathUse::Display;
-				AssetDefinition->GetSourceFiles(GetSourceFilesArgs, [&OutUniqueSourceFileLabels](const FAssetSourceFilesResult& AssetImportInfo)
-				{
-					OutUniqueSourceFileLabels.Add(AssetImportInfo.FilePath);
+					OutUniqueSourceFileLabels.AddUnique(AssetImportInfo.DisplayLabel);
 					return true;
 				});
 
diff --git a/Engine/Source/Editor/AssetDefinition/Private/AssetDefinition.cpp b/Engine/Source/Editor/AssetDefinition/Private/AssetDefinition.cpp
--- a/Engine/Source/Editor/AssetDefinition/Private/AssetDefinition.cpp
+++ b/Engine/Source/Editor/AssetDefinition/Private/AssetDefinition.cpp
@@ -95,6 +95,7 @@ EAssetCommandResult UAssetDefinition::GetSourceFiles(const FAssetSourceFilesArgs
 				for (FAssetImportInfo::FSourceFile& SourceFiles : ImportInfo.SourceFiles)
 				{
 					Result.FilePath = MoveTemp(SourceFiles.RelativeFilename);
+					Result.DisplayLabel = MoveTemp(SourceFiles.DisplayLabelName);
 					Result.Timestamp = MoveTemp(SourceFiles.Timestamp);
 					Result.FileHash = MoveTemp(SourceFiles.FileHash);
 				
diff --git a/Engine/Source/Editor/AssetDefinition/Public/AssetDefinition.h b/Engine/Source/Editor/AssetDefinition/Public/AssetDefinition.h
--- a/Engine/Source/Editor/AssetDefinition/Public/AssetDefinition.h
+++ b/Engine/Source/Editor/AssetDefinition/Public/AssetDefinition.h
@@ -141,6 +141,9 @@ struct FAssetSourceFilesResult
 	/** The file path in the format requested. */
 	FString FilePath;
 
+	/** The Label was used to display this source file in the property editor. */
+	FString DisplayLabel;
+
 	/** The timestamp of the file when it was imported (as UTC). 0 when unknown. */
 	FDateTime Timestamp;
 
