Change content browser asset registry callbacks to ignore directory paths like __ExternalActors__ that can never be searched for. This significantly improves the initial asset registry scan performance

diff --git a/Engine/Plugins/Editor/ContentBrowser/ContentBrowserAssetDataSource/Source/ContentBrowserAssetDataSource/Private/ContentBrowserAssetDataSource.cpp b/Engine/Plugins/Editor/ContentBrowser/ContentBrowserAssetDataSource/Source/ContentBrowserAssetDataSource/Private/ContentBrowserAssetDataSource.cpp
--- a/Engine/Plugins/Editor/ContentBrowser/ContentBrowserAssetDataSource/Source/ContentBrowserAssetDataSource/Private/ContentBrowserAssetDataSource.cpp
+++ b/Engine/Plugins/Editor/ContentBrowser/ContentBrowserAssetDataSource/Source/ContentBrowserAssetDataSource/Private/ContentBrowserAssetDataSource.cpp
@@ -2859,6 +2859,12 @@ void UContentBrowserAssetDataSource::OnObjectPreSave(UObject* InObject, FObjectP
 
 void UContentBrowserAssetDataSource::OnPathAdded(const FString& InPath)
 {
+	// Completely ignore paths that do not pass the most inclusive filter
+	if (!ContentBrowserDataUtils::PathPassesAttributeFilter(InPath, 0, EContentBrowserItemAttributeFilter::IncludeAll))
+	{
+		return;
+	}
+
 	FName PathName(InPath);
 	RecentlyPopulatedAssetFolders.Empty();
 	
