Enable plugins for use with LiveLinkHub program

diff --git a/Engine/Plugins/Editor/ContentBrowser/ContentBrowserAssetDataSource/ContentBrowserAssetDataSource.uplugin b/Engine/Plugins/Editor/ContentBrowser/ContentBrowserAssetDataSource/ContentBrowserAssetDataSource.uplugin
--- a/Engine/Plugins/Editor/ContentBrowser/ContentBrowserAssetDataSource/ContentBrowserAssetDataSource.uplugin
+++ b/Engine/Plugins/Editor/ContentBrowser/ContentBrowserAssetDataSource/ContentBrowserAssetDataSource.uplugin
@@ -13,13 +13,15 @@
 	"EnabledByDefault" : true,
 	"CanContainContent" : false,
 	"IsBetaVersion" : false,
-	"Installed" : false,
+	"Installed": false,
+	"SupportedPrograms":  ["LiveLinkHubLauncher"],
 	"Modules" :
 	[
 		{
-			"Name" : "ContentBrowserAssetDataSource",
-			"Type" : "Editor",
-			"LoadingPhase" : "Default"
+			"Name": "ContentBrowserAssetDataSource",
+			"Type": "Editor",
+			"LoadingPhase": "Default",
+			"ProgramAllowlist": [ "LiveLinkHubLauncher" ]
 		}
 	]
 }
\ No newline at end of file
diff --git a/Engine/Plugins/Experimental/StructUtils/StructUtils.uplugin b/Engine/Plugins/Experimental/StructUtils/StructUtils.uplugin
--- a/Engine/Plugins/Experimental/StructUtils/StructUtils.uplugin
+++ b/Engine/Plugins/Experimental/StructUtils/StructUtils.uplugin
@@ -15,11 +15,13 @@
 	"IsBetaVersion": false,
 	"IsExperimentalVersion": true,
 	"Installed": false,
+	"SupportedPrograms": [ "LiveLinkHubLauncher" ],
 	"Modules": [
 		{
 			"Name": "StructUtils",
 			"Type": "Runtime",
-			"LoadingPhase": "PreDefault"
+			"LoadingPhase": "PreDefault",
+			"ProgramAllowList": [ "LiveLinkHubLauncher" ]
 		},
 		{
 			"Name": "StructUtilsEngine",
diff --git a/Engine/Plugins/Messaging/UdpMessaging/UdpMessaging.uplugin b/Engine/Plugins/Messaging/UdpMessaging/UdpMessaging.uplugin
--- a/Engine/Plugins/Messaging/UdpMessaging/UdpMessaging.uplugin
+++ b/Engine/Plugins/Messaging/UdpMessaging/UdpMessaging.uplugin
@@ -14,14 +14,15 @@
 	"CanContainContent" : false,
 	"IsBetaVersion" : false,
 	"Installed" : false,
-	"SupportedPrograms" : 
-	[
+	"SupportedPrograms": [
 		"UnrealFrontend",
 		"UnrealMultiUserServer",
 		"UnrealMultiUserSlateServer",
 		"UnrealRecoverySvc",
 		"UnrealLightmass",
-		"CrashReportClientEditor"
+		"CrashReportClientEditor",
+		"LiveLinkHubLauncher",
+		"TimeSynchronizationWriter"
 	],
 	"Modules" :
 	[
@@ -38,7 +39,9 @@
 				"UnrealMultiUserSlateServer",
 				"UnrealRecoverySvc",
 				"UnrealLightmass",
-				"CrashReportClientEditor"
+				"CrashReportClientEditor",
+				"LiveLinkHubLauncher",
+				"TimeSynchronizationWriter"
 			]
 		}
 	]
diff --git a/Engine/Plugins/VirtualProduction/LiveLinkFreeD/LiveLinkFreeD.uplugin b/Engine/Plugins/VirtualProduction/LiveLinkFreeD/LiveLinkFreeD.uplugin
--- a/Engine/Plugins/VirtualProduction/LiveLinkFreeD/LiveLinkFreeD.uplugin
+++ b/Engine/Plugins/VirtualProduction/LiveLinkFreeD/LiveLinkFreeD.uplugin
@@ -15,6 +15,7 @@
 	"IsExperimentalVersion": false,
 	"Installed": false,
 	"EnabledByDefault": false,
+	"SupportedPrograms": [ "LiveLinkHubLauncher" ],
 	"Plugins": [
 		{
 			"Name": "UdpMessaging",
@@ -25,7 +26,9 @@
 		{
 			"Name": "LiveLinkFreeD",
 			"Type": "Runtime",
-			"LoadingPhase": "Default"
+			"LoadingPhase": "Default",
+			"ProgramAllowList": [ "LiveLinkHubLauncher" ]
 		}
 	]
 }
+ 
\ No newline at end of file
