Rename LiveLinkHubLauncher to LiveLinkHub

diff --git a/Engine/Plugins/Editor/ContentBrowser/ContentBrowserAssetDataSource/ContentBrowserAssetDataSource.uplugin b/Engine/Plugins/Editor/ContentBrowser/ContentBrowserAssetDataSource/ContentBrowserAssetDataSource.uplugin
--- a/Engine/Plugins/Editor/ContentBrowser/ContentBrowserAssetDataSource/ContentBrowserAssetDataSource.uplugin
+++ b/Engine/Plugins/Editor/ContentBrowser/ContentBrowserAssetDataSource/ContentBrowserAssetDataSource.uplugin
@@ -14,14 +14,14 @@
 	"CanContainContent" : false,
 	"IsBetaVersion" : false,
 	"Installed": false,
-	"SupportedPrograms":  ["LiveLinkHubLauncher"],
+	"SupportedPrograms":  ["LiveLinkHub"],
 	"Modules" :
 	[
 		{
 			"Name": "ContentBrowserAssetDataSource",
 			"Type": "Editor",
 			"LoadingPhase": "Default",
-			"ProgramAllowlist": [ "LiveLinkHubLauncher" ]
+			"ProgramAllowlist": [ "LiveLinkHub" ]
 		}
 	]
 }
\ No newline at end of file
diff --git a/Engine/Plugins/Experimental/QuicMessaging/QuicMessaging.uplugin b/Engine/Plugins/Experimental/QuicMessaging/QuicMessaging.uplugin
--- a/Engine/Plugins/Experimental/QuicMessaging/QuicMessaging.uplugin
+++ b/Engine/Plugins/Experimental/QuicMessaging/QuicMessaging.uplugin
@@ -26,7 +26,8 @@
 		"UnrealRecoverySvc",
 		"UnrealLightmass",
 		"CrashReportClientEditor",
-		"CoopMultiUserServer"
+		"CoopMultiUserServer",
+		"LiveLinkHub"
 	],
 	"Modules": [
 		{
@@ -41,7 +42,8 @@
 				"UnrealRecoverySvc",
 				"UnrealLightmass",
 				"CrashReportClientEditor",
-				"CoopMultiUserServer"
+				"CoopMultiUserServer",
+				"LiveLinkHub"
 			]
 		},
 		{
@@ -56,7 +58,8 @@
 				"UnrealRecoverySvc",
 				"UnrealLightmass",
 				"CrashReportClientEditor",
-				"CoopMultiUserServer"
+				"CoopMultiUserServer",
+				"LiveLinkHub"
 			]
 		}
 	],
diff --git a/Engine/Plugins/Experimental/StructUtils/StructUtils.uplugin b/Engine/Plugins/Experimental/StructUtils/StructUtils.uplugin
--- a/Engine/Plugins/Experimental/StructUtils/StructUtils.uplugin
+++ b/Engine/Plugins/Experimental/StructUtils/StructUtils.uplugin
@@ -15,13 +15,13 @@
 	"IsBetaVersion": false,
 	"IsExperimentalVersion": true,
 	"Installed": false,
-	"SupportedPrograms": [ "LiveLinkHubLauncher" ],
+	"SupportedPrograms": [ "LiveLinkHub" ],
 	"Modules": [
 		{
 			"Name": "StructUtils",
 			"Type": "Runtime",
 			"LoadingPhase": "PreDefault",
-			"ProgramAllowList": [ "LiveLinkHubLauncher" ]
+			"ProgramAllowList": [ "LiveLinkHub" ]
 		},
 		{
 			"Name": "StructUtilsEngine",
diff --git a/Engine/Plugins/Messaging/UdpMessaging/UdpMessaging.uplugin b/Engine/Plugins/Messaging/UdpMessaging/UdpMessaging.uplugin
--- a/Engine/Plugins/Messaging/UdpMessaging/UdpMessaging.uplugin
+++ b/Engine/Plugins/Messaging/UdpMessaging/UdpMessaging.uplugin
@@ -21,7 +21,7 @@
 		"UnrealRecoverySvc",
 		"UnrealLightmass",
 		"CrashReportClientEditor",
-		"LiveLinkHubLauncher",
+		"LiveLinkHub",
 		"TimeSynchronizationWriter"
 	],
 	"Modules" :
@@ -40,7 +40,7 @@
 				"UnrealRecoverySvc",
 				"UnrealLightmass",
 				"CrashReportClientEditor",
-				"LiveLinkHubLauncher",
+				"LiveLinkHub",
 				"TimeSynchronizationWriter"
 			]
 		}
diff --git a/Engine/Plugins/VirtualProduction/LiveLinkFreeD/LiveLinkFreeD.uplugin b/Engine/Plugins/VirtualProduction/LiveLinkFreeD/LiveLinkFreeD.uplugin
--- a/Engine/Plugins/VirtualProduction/LiveLinkFreeD/LiveLinkFreeD.uplugin
+++ b/Engine/Plugins/VirtualProduction/LiveLinkFreeD/LiveLinkFreeD.uplugin
@@ -15,7 +15,7 @@
 	"IsExperimentalVersion": false,
 	"Installed": false,
 	"EnabledByDefault": false,
-	"SupportedPrograms": [ "LiveLinkHubLauncher" ],
+	"SupportedPrograms": [ "LiveLinkHub" ],
 	"Plugins": [
 		{
 			"Name": "UdpMessaging",
@@ -27,7 +27,7 @@
 			"Name": "LiveLinkFreeD",
 			"Type": "Runtime",
 			"LoadingPhase": "Default",
-			"ProgramAllowList": [ "LiveLinkHubLauncher" ]
+			"ProgramAllowList": [ "LiveLinkHub" ]
 		}
 	]
 }
