Added tag metadata to items in New Asset menu of Context Browser

diff --git a/Engine/Plugins/Editor/ContentBrowser/ContentBrowserAssetDataSource/Source/ContentBrowserAssetDataSource/Private/NewAssetContextMenu.cpp b/Engine/Plugins/Editor/ContentBrowser/ContentBrowserAssetDataSource/Source/ContentBrowserAssetDataSource/Private/NewAssetContextMenu.cpp
--- a/Engine/Plugins/Editor/ContentBrowser/ContentBrowserAssetDataSource/Source/ContentBrowserAssetDataSource/Private/NewAssetContextMenu.cpp
+++ b/Engine/Plugins/Editor/ContentBrowser/ContentBrowserAssetDataSource/Source/ContentBrowserAssetDataSource/Private/NewAssetContextMenu.cpp
@@ -389,13 +389,25 @@ void FNewAssetContextMenu::CreateNewAssetMenus(UToolMenu* Menu, FName SectionNam
 	{
 		TWeakObjectPtr<UClass> WeakFactoryClass = FactoryItem.Factory->GetClass();
 
+		FName AssetTypeName;
+
+		if (UFactory* Factory = FactoryItem.Factory)
+		{
+			if (UClass* SupportedClass = Factory->GetSupportedClass())
+			{
+				AssetTypeName = SupportedClass->GetFName();
+			}
+		}
+
 		Section.AddEntry(FToolMenuEntry::InitMenuEntry(
 			NAME_None,
 			FUIAction(
 				FExecuteAction::CreateStatic(&FNewAssetContextMenu::ExecuteNewAsset, InOnNewAssetRequested, InPath, WeakFactoryClass),
 				InCanExecuteAction
 			),
-			SNew(SFactoryMenuEntry, FactoryItem.Factory)));
+			SNew(SFactoryMenuEntry, FactoryItem.Factory)
+				.AddMetaData<FTagMetaData>(FTagMetaData(AssetTypeName))));
+
 	}
 
 	if (SubMenuData->Children.Num() == 0)
diff --git a/Engine/Source/Editor/ContentBrowser/Private/SContentBrowser.cpp b/Engine/Source/Editor/ContentBrowser/Private/SContentBrowser.cpp
--- a/Engine/Source/Editor/ContentBrowser/Private/SContentBrowser.cpp
+++ b/Engine/Source/Editor/ContentBrowser/Private/SContentBrowser.cpp
@@ -2844,7 +2844,9 @@ TSharedRef<SWidget> SContentBrowser::MakeAddNewContextMenu(const EContentBrowser
 	FToolMenuContext ToolMenuContext(nullptr, MenuExtender, nullptr);
 	AppendNewMenuContextObjects(InDomain, SourcesData.VirtualPaths, ToolMenuContext, CommonContext, bCanBeModified);
 
-	return UToolMenus::Get()->GenerateWidget("ContentBrowser.AddNewContextMenu", ToolMenuContext);
+	TSharedRef<SWidget> GeneratedWidget = UToolMenus::Get()->GenerateWidget("ContentBrowser.AddNewContextMenu", ToolMenuContext);
+	GeneratedWidget->AddMetadata<FTagMetaData>(MakeShared<FTagMetaData>(TEXT("ContentBrowser.AddNewContextMenu")));
+	return GeneratedWidget;
 }
 
 void SContentBrowser::PopulateAddNewContextMenu(class UToolMenu* Menu)
