Static Analysis warning fix - used reference instead of a pointer for a non-nullable variable

diff --git a/Engine/Plugins/Editor/ContentBrowser/ContentBrowserAssetDataSource/Source/ContentBrowserAssetDataSource/Private/NewAssetContextMenu.cpp b/Engine/Plugins/Editor/ContentBrowser/ContentBrowserAssetDataSource/Source/ContentBrowserAssetDataSource/Private/NewAssetContextMenu.cpp
--- a/Engine/Plugins/Editor/ContentBrowser/ContentBrowserAssetDataSource/Source/ContentBrowserAssetDataSource/Private/NewAssetContextMenu.cpp
+++ b/Engine/Plugins/Editor/ContentBrowser/ContentBrowserAssetDataSource/Source/ContentBrowserAssetDataSource/Private/NewAssetContextMenu.cpp
@@ -22,10 +22,10 @@
 
 struct FFactoryItem
 {
-	UFactory* Factory;
+	UFactory& Factory;
 	FText DisplayName;
 
-	FFactoryItem(UFactory* InFactory, const FText& InDisplayName)
+	FFactoryItem(UFactory& InFactory, const FText& InDisplayName)
 		: Factory(InFactory)
 		, DisplayName(InDisplayName)
 	{
@@ -93,7 +93,7 @@ static TArray<FFactoryItem> FindFactoriesInCategory(const TArray<UFactory*>& Fac
 		const uint32 FactoryCategories = Factory->GetMenuCategories();
 		if (FactoryCategories & AssetTypeCategory)
 		{
-			FactoriesInThisCategory.Emplace(Factory, Factory->GetDisplayName());
+			FactoriesInThisCategory.Emplace(*Factory, Factory->GetDisplayName());
 
 			if (FindFirstOnly)
 			{
@@ -356,7 +356,7 @@ void FNewAssetContextMenu::CreateNewAssetMenuCategory(UToolMenu* Menu, FName Sec
 	for (FFactoryItem& Item : FactoriesInThisCategory)
 	{
 		FCategorySubMenuItem* SubMenu = ParentMenuData.Get();
-		const TArray<FText>& CategoryNames = Item.Factory->GetMenuCategorySubMenus();
+		const TArray<FText>& CategoryNames = Item.Factory.GetMenuCategorySubMenus();
 		for (FText CategoryName : CategoryNames)
 		{
 			const FString SourceString = CategoryName.BuildSourceString();
@@ -387,16 +387,13 @@ void FNewAssetContextMenu::CreateNewAssetMenus(UToolMenu* Menu, FName SectionNam
 	FToolMenuSection& Section = Menu->FindOrAddSection(SectionName);
 	for (const FFactoryItem& FactoryItem : SubMenuData->Factories)
 	{
-		TWeakObjectPtr<UClass> WeakFactoryClass = FactoryItem.Factory->GetClass();
+		TWeakObjectPtr<UClass> WeakFactoryClass = FactoryItem.Factory.GetClass();
 
 		FName AssetTypeName;
 
-		if (UFactory* Factory = FactoryItem.Factory)
+		if (UClass* SupportedClass = FactoryItem.Factory.GetSupportedClass())
 		{
-			if (UClass* SupportedClass = Factory->GetSupportedClass())
-			{
-				AssetTypeName = SupportedClass->GetFName();
-			}
+			AssetTypeName = SupportedClass->GetFName();
 		}
 
 		Section.AddEntry(FToolMenuEntry::InitMenuEntry(
@@ -405,7 +402,7 @@ void FNewAssetContextMenu::CreateNewAssetMenus(UToolMenu* Menu, FName SectionNam
 				FExecuteAction::CreateStatic(&FNewAssetContextMenu::ExecuteNewAsset, InOnNewAssetRequested, InPath, WeakFactoryClass),
 				InCanExecuteAction
 			),
-			SNew(SFactoryMenuEntry, FactoryItem.Factory)
+			SNew(SFactoryMenuEntry, &FactoryItem.Factory)
 				.AddMetaData<FTagMetaData>(FTagMetaData(AssetTypeName))));
 
 	}
