Make slow task immediately for opening assets to improve sense of responsiveness

diff --git a/Engine/Plugins/Editor/ContentBrowser/ContentBrowserAssetDataSource/Source/ContentBrowserAssetDataSource/Private/ContentBrowserAssetDataCore.cpp b/Engine/Plugins/Editor/ContentBrowser/ContentBrowserAssetDataSource/Source/ContentBrowserAssetDataSource/Private/ContentBrowserAssetDataCore.cpp
--- a/Engine/Plugins/Editor/ContentBrowser/ContentBrowserAssetDataSource/Source/ContentBrowserAssetDataSource/Private/ContentBrowserAssetDataCore.cpp
+++ b/Engine/Plugins/Editor/ContentBrowser/ContentBrowserAssetDataSource/Source/ContentBrowserAssetDataSource/Private/ContentBrowserAssetDataCore.cpp
@@ -321,6 +321,7 @@ bool EditOrPreviewAssetFileItems(TArrayView<const TSharedRef<const FContentBrows
 		: FText::Format(LOCTEXT("LoadingXAssets", "Loading {0} {0}|plural(one=Asset,other=Assets)..."), InAssetPayloads.Num());
 
 	FScopedSlowTask SlowTask(100, DefaultText);
+	SlowTask.MakeDialogDelayed(0.1f);
 
 	// Iterate over all activated assets to map them to AssetTypeActions.
 	// This way individual asset type actions will get a batched list of assets to operate on
@@ -351,11 +352,6 @@ bool EditOrPreviewAssetFileItems(TArrayView<const TSharedRef<const FContentBrows
 		
 		for (const FAssetData& AssetData : AssetsToLoad)
 		{
-			if (!AssetData.IsAssetLoaded() && FEditorFileUtils::IsMapPackageAsset(AssetData.GetObjectPathString()))
-			{
-				SlowTask.MakeDialog();
-			}
-
 			SlowTask.EnterProgressFrame(75.0f / InAssetPayloads.Num(), FText::Format(LOCTEXT("LoadingAssetName", "Loading {0}..."), FText::FromName(AssetData.AssetName)));
 
 			ObjList.Add(AssetData.GetAsset());
