Add plugin category flag EContentBrowserItemFlags::Category_Plugin
Assign flag to content browser folders and assets in plugins.
Fix incorrect early return in UContentBrowserAssetDataSource::OnPathsAdded

diff --git a/Engine/Plugins/Editor/ContentBrowser/ContentBrowserAssetDataSource/Source/ContentBrowserAssetDataSource/Private/ContentBrowserAssetDataCore.cpp b/Engine/Plugins/Editor/ContentBrowser/ContentBrowserAssetDataSource/Source/ContentBrowserAssetDataSource/Private/ContentBrowserAssetDataCore.cpp
--- a/Engine/Plugins/Editor/ContentBrowser/ContentBrowserAssetDataSource/Source/ContentBrowserAssetDataSource/Private/ContentBrowserAssetDataCore.cpp
+++ b/Engine/Plugins/Editor/ContentBrowser/ContentBrowserAssetDataSource/Source/ContentBrowserAssetDataSource/Private/ContentBrowserAssetDataCore.cpp
@@ -30,21 +30,31 @@
 
 #define LOCTEXT_NAMESPACE "ContentBrowserAssetDataSource"
 
-DEFINE_LOG_CATEGORY_STATIC(LogContentBrowserAssetDataSource, Warning, Warning);
+DEFINE_LOG_CATEGORY(LogContentBrowserAssetDataSource);
 
 namespace ContentBrowserAssetData
 {
 
-FContentBrowserItemData CreateAssetFolderItem(UContentBrowserDataSource* InOwnerDataSource, const FName InVirtualPath, const FName InFolderPath, const bool bIsCookedPath)
+FContentBrowserItemData CreateAssetFolderItem(UContentBrowserDataSource* InOwnerDataSource, const FName InVirtualPath, const FName InFolderPath, const bool bIsCookedPath, const bool bIsPlugin)
 {
 	const FString FolderItemName = FPackageName::GetShortName(InFolderPath);
 	FText FolderDisplayNameOverride = ContentBrowserDataUtils::GetFolderItemDisplayNameOverride(InFolderPath, FolderItemName, /*bIsClassesFolder*/ false, bIsCookedPath);
-	return FContentBrowserItemData(InOwnerDataSource, EContentBrowserItemFlags::Type_Folder | EContentBrowserItemFlags::Category_Asset, InVirtualPath, *FolderItemName, MoveTemp(FolderDisplayNameOverride), MakeShared<FContentBrowserAssetFolderItemDataPayload>(InFolderPath));
-}
-
-FContentBrowserItemData CreateAssetFileItem(UContentBrowserDataSource* InOwnerDataSource, const FName InVirtualPath, const FAssetData& InAssetData)
-{
-	return FContentBrowserItemData(InOwnerDataSource, EContentBrowserItemFlags::Type_File | EContentBrowserItemFlags::Category_Asset, InVirtualPath, InAssetData.AssetName, FText(), MakeShared<FContentBrowserAssetFileItemDataPayload>(InAssetData));
+	return FContentBrowserItemData(InOwnerDataSource,
+		EContentBrowserItemFlags::Type_Folder | EContentBrowserItemFlags::Category_Asset | (bIsPlugin ? EContentBrowserItemFlags::Category_Plugin : EContentBrowserItemFlags::None ),
+		InVirtualPath,
+		*FolderItemName,
+		MoveTemp(FolderDisplayNameOverride),
+		MakeShared<FContentBrowserAssetFolderItemDataPayload>(InFolderPath));
+}
+
+FContentBrowserItemData CreateAssetFileItem(UContentBrowserDataSource* InOwnerDataSource, const FName InVirtualPath, const FAssetData& InAssetData, const bool bIsPluginAsset)
+{
+	return FContentBrowserItemData(InOwnerDataSource,
+		EContentBrowserItemFlags::Type_File | EContentBrowserItemFlags::Category_Asset | (bIsPluginAsset ? EContentBrowserItemFlags::Category_Plugin : EContentBrowserItemFlags::None),
+		InVirtualPath,
+		InAssetData.AssetName,
+		FText(),
+		MakeShared<FContentBrowserAssetFileItemDataPayload>(InAssetData));
 }
 
 FContentBrowserItemData CreateUnsupportedAssetFileItem(UContentBrowserDataSource* InOwnerDataSource, const FName InVirtualPath, const FAssetData& InAssetData)
diff --git a/Engine/Plugins/Editor/ContentBrowser/ContentBrowserAssetDataSource/Source/ContentBrowserAssetDataSource/Private/ContentBrowserAssetDataSource.cpp b/Engine/Plugins/Editor/ContentBrowser/ContentBrowserAssetDataSource/Source/ContentBrowserAssetDataSource/Private/ContentBrowserAssetDataSource.cpp
--- a/Engine/Plugins/Editor/ContentBrowser/ContentBrowserAssetDataSource/Source/ContentBrowserAssetDataSource/Private/ContentBrowserAssetDataSource.cpp
+++ b/Engine/Plugins/Editor/ContentBrowser/ContentBrowserAssetDataSource/Source/ContentBrowserAssetDataSource/Private/ContentBrowserAssetDataSource.cpp
@@ -4,6 +4,7 @@
 #include "ContentBrowserAssetDataCore.h"
 #include "AssetRegistry/AssetRegistryModule.h"
 #include "AssetToolsModule.h"
+#include "AssetViewUtils.h"
 #include "CollectionManagerModule.h"
 #include "ContentBrowserAssetDataPayload.h"
 #include "ICollectionManager.h"
@@ -151,6 +152,11 @@ void UContentBrowserAssetDataSource::Initialize(const bool InAutoRegister)
 
 	// Populate the initial set of folder attributes
 	// This will be updated as the scan finds more content
+	AssetRegistry->EnumerateAllCachedPaths([this](FName PathName) { 
+		FNameBuilder NameBuilder{PathName};
+		OnPathsAdded({NameBuilder.ToView()});
+		return true; 
+	});
 	AssetRegistry->EnumerateAllAssets([this](const FAssetData& InAssetData)
 		{
 			OnPathPopulated(InAssetData);
@@ -2718,8 +2724,8 @@ FContentBrowserItemData UContentBrowserAssetDataSource::CreateAssetFolderItem(co
 
 	const EContentBrowserFolderAttributes FolderAttributes = GetAssetFolderAttributes(InFolderPath);
 	const bool bIsCookedPath = EnumHasAnyFlags(FolderAttributes, EContentBrowserFolderAttributes::HasContent) && !EnumHasAnyFlags(FolderAttributes, EContentBrowserFolderAttributes::HasSourceContent);
-
-	return ContentBrowserAssetData::CreateAssetFolderItem(this, VirtualizedPath, InFolderPath, bIsCookedPath);
+	const bool bIsPlugin = EnumHasAnyFlags(FolderAttributes, EContentBrowserFolderAttributes::IsInPlugin);
+	return ContentBrowserAssetData::CreateAssetFolderItem(this, VirtualizedPath, InFolderPath, bIsCookedPath, bIsPlugin);
 }
 
 FContentBrowserItemData UContentBrowserAssetDataSource::CreateAssetFileItem(const FAssetData& InAssetData)
@@ -2729,7 +2735,9 @@ PRAGMA_DISABLE_DEPRECATION_WARNINGS
 	TryConvertInternalPathToVirtual(InAssetData.ObjectPath, VirtualizedPath);
 PRAGMA_ENABLE_DEPRECATION_WARNINGS
 
-	return ContentBrowserAssetData::CreateAssetFileItem(this, VirtualizedPath, InAssetData);
+	const EContentBrowserFolderAttributes FolderAttributes = GetAssetFolderAttributes(InAssetData.PackagePath);
+	const bool bIsPlugin = EnumHasAnyFlags(FolderAttributes, EContentBrowserFolderAttributes::IsInPlugin);
+	return ContentBrowserAssetData::CreateAssetFileItem(this, VirtualizedPath, InAssetData, bIsPlugin);
 }
 
 FContentBrowserItemData UContentBrowserAssetDataSource::CreateUnsupportedAssetFileItem(const FAssetData& InAssetData)
@@ -2859,31 +2867,36 @@ void UContentBrowserAssetDataSource::OnObjectPreSave(UObject* InObject, FObjectP
 
 void UContentBrowserAssetDataSource::OnPathsAdded(TConstArrayView<FStringView> Paths)
 {
+	RecentlyPopulatedAssetFolders.Empty();
 	for (FStringView InPath : Paths)
 	{
 		// Completely ignore paths that do not pass the most inclusive filter
 		if (!ContentBrowserDataUtils::PathPassesAttributeFilter(InPath, 0, EContentBrowserItemAttributeFilter::IncludeAll))
 		{
-			return;
+			continue;
 		}
 
 		FName PathName(InPath);
-		RecentlyPopulatedAssetFolders.Empty();
-		
+		const bool bIsPlugin = AssetViewUtils::IsPluginFolder(InPath);
+		if (bIsPlugin)
+		{
+			OnPathPopulated(InPath, EContentBrowserFolderAttributes::IsInPlugin);
+		}
+
 		QueueItemDataUpdate(FContentBrowserItemDataUpdate::MakeItemAddedUpdate(CreateAssetFolderItem(PathName)));
 
-		FStringView PathView(InPath);
 		// Minus one because the test depth start at zero
-		const int32 CurrentDepth = ContentBrowserDataUtils::CalculateFolderDepthOfPath(PathView) - 1;
+		const int32 CurrentDepth = ContentBrowserDataUtils::CalculateFolderDepthOfPath(InPath) - 1;
 		int32 Index;
-		if (PathView.FindLastChar(TEXT('/'), Index))
+		if (InPath.FindLastChar(TEXT('/'), Index))
 		{ 
 			uint32 PathNameHash = GetTypeHash(PathName);
-			FName ParentPath(PathView.Left(Index));
+			FName ParentPath(InPath.Left(Index));
 			uint32 ParentPathHash = GetTypeHash(ParentPath);
-			OnAssetPathAddedDelegate.Broadcast(PathName, PathView, PathNameHash, ParentPath, ParentPathHash, CurrentDepth);
+			OnAssetPathAddedDelegate.Broadcast(PathName, InPath, PathNameHash, ParentPath, ParentPathHash, CurrentDepth);
 		}
 	}
+	RecentlyPopulatedAssetFolders.Empty();
 }
 
 void UContentBrowserAssetDataSource::OnPathsRemoved(TConstArrayView<FStringView> Paths)
diff --git a/Engine/Plugins/Editor/ContentBrowser/ContentBrowserAssetDataSource/Source/ContentBrowserAssetDataSource/Public/ContentBrowserAssetDataCore.h b/Engine/Plugins/Editor/ContentBrowser/ContentBrowserAssetDataSource/Source/ContentBrowserAssetDataSource/Public/ContentBrowserAssetDataCore.h
--- a/Engine/Plugins/Editor/ContentBrowser/ContentBrowserAssetDataSource/Source/ContentBrowserAssetDataSource/Public/ContentBrowserAssetDataCore.h
+++ b/Engine/Plugins/Editor/ContentBrowser/ContentBrowserAssetDataSource/Source/ContentBrowserAssetDataSource/Public/ContentBrowserAssetDataCore.h
@@ -17,12 +17,14 @@ class UObject;
 class UToolMenu;
 class UContentBrowserDataSource;
 
+DECLARE_LOG_CATEGORY_EXTERN(LogContentBrowserAssetDataSource, Log, All);
+
 namespace ContentBrowserAssetData
 {
 
-	CONTENTBROWSERASSETDATASOURCE_API FContentBrowserItemData CreateAssetFolderItem(UContentBrowserDataSource* InOwnerDataSource, const FName InVirtualPath, const FName InFolderPath, const bool bIsCookedPath = false);
+	CONTENTBROWSERASSETDATASOURCE_API FContentBrowserItemData CreateAssetFolderItem(UContentBrowserDataSource* InOwnerDataSource, const FName InVirtualPath, const FName InFolderPath, const bool bIsCookedPath = false, const bool bIsPlugin = false);
 
-	CONTENTBROWSERASSETDATASOURCE_API FContentBrowserItemData CreateAssetFileItem(UContentBrowserDataSource* InOwnerDataSource, const FName InVirtualPath, const FAssetData& InAssetData);
+	CONTENTBROWSERASSETDATASOURCE_API FContentBrowserItemData CreateAssetFileItem(UContentBrowserDataSource* InOwnerDataSource, const FName InVirtualPath, const FAssetData& InAssetData, const bool bIsPlugin = false);
 
 	CONTENTBROWSERASSETDATASOURCE_API FContentBrowserItemData CreateUnsupportedAssetFileItem(UContentBrowserDataSource* InOwnerDataSource, const FName InVirtualPath, const FAssetData& InAssetData);
 
diff --git a/Engine/Plugins/Editor/ContentBrowser/ContentBrowserAssetDataSource/Source/ContentBrowserAssetDataSource/Public/ContentBrowserAssetDataSource.h b/Engine/Plugins/Editor/ContentBrowser/ContentBrowserAssetDataSource/Source/ContentBrowserAssetDataSource/Public/ContentBrowserAssetDataSource.h
--- a/Engine/Plugins/Editor/ContentBrowser/ContentBrowserAssetDataSource/Source/ContentBrowserAssetDataSource/Public/ContentBrowserAssetDataSource.h
+++ b/Engine/Plugins/Editor/ContentBrowser/ContentBrowserAssetDataSource/Source/ContentBrowserAssetDataSource/Public/ContentBrowserAssetDataSource.h
@@ -103,6 +103,11 @@ enum class EContentBrowserFolderAttributes : uint8
 	 * This folder has source (uncooked) content that will appear in the Content Browser view.
 	 */
 	HasSourceContent = 1<<3,
+
+	/** 
+	 * This folder is inside a plugin.
+	 */
+	IsInPlugin = 1<<4,
 };
 ENUM_CLASS_FLAGS(EContentBrowserFolderAttributes);
 
diff --git a/Engine/Source/Editor/ContentBrowserData/Private/ContentBrowserItemData.cpp b/Engine/Source/Editor/ContentBrowserData/Private/ContentBrowserItemData.cpp
--- a/Engine/Source/Editor/ContentBrowserData/Private/ContentBrowserItemData.cpp
+++ b/Engine/Source/Editor/ContentBrowserData/Private/ContentBrowserItemData.cpp
@@ -48,6 +48,11 @@ bool FContentBrowserItemData::IsFile() const
 	return EnumHasAnyFlags(ItemFlags, EContentBrowserItemFlags::Type_File);
 }
 
+bool FContentBrowserItemData::IsPlugin() const
+{
+	return EnumHasAnyFlags(ItemFlags, EContentBrowserItemFlags::Category_Plugin);
+}
+
 bool FContentBrowserItemData::IsSupported() const
 {
 	return !EnumHasAnyFlags(ItemFlags, EContentBrowserItemFlags::Misc_Unsupported);
diff --git a/Engine/Source/Editor/ContentBrowserData/Public/ContentBrowserItemData.h b/Engine/Source/Editor/ContentBrowserData/Public/ContentBrowserItemData.h
--- a/Engine/Source/Editor/ContentBrowserData/Public/ContentBrowserItemData.h
+++ b/Engine/Source/Editor/ContentBrowserData/Public/ContentBrowserItemData.h
@@ -37,16 +37,17 @@ enum class EContentBrowserItemFlags : uint16
 	Category_Asset = 1<<2,
 	Category_Class = 1<<3,
 	Category_Collection = 1<<4,
-	Category_Misc = 1<<5,
-	Category_MASK = Category_Asset | Category_Class | Category_Collection | Category_Misc,
+	Category_Plugin = 1 << 5, // Plugin content or classes - note that folders containing plugins do not have this flag set 
+	Category_Misc = 1<<6,
+	Category_MASK = Category_Asset | Category_Class | Category_Collection | Category_Plugin | Category_Misc,
 
 	/** Flags denoting additional information for temporary items (mutually exclusive) */
-	Temporary_Creation = 1<<6,
-	Temporary_Duplication = 1<<7,
+	Temporary_Creation = 1<<7,
+	Temporary_Duplication = 1<<8,
 	Temporary_MASK = Temporary_Creation | Temporary_Duplication,
 
 	/** Flag to mark the item as showing something that is unsupported */
-	Misc_Unsupported = 1 << 8,
+	Misc_Unsupported = 1 << 9,
 };
 ENUM_CLASS_FLAGS(EContentBrowserItemFlags);
 
@@ -168,6 +169,12 @@ public:
 	 */
 	bool IsFile() const;
 
+	/**
+	 * Check to see whether this item is in a plugin. Folders and files inside plugins return true, folders which
+	 * contain and organize plugins do not.
+	 */
+	bool IsPlugin() const;
+
 	/**
 	 * Check if the item is representing a supported item
 	 * The content browser can also display some unsupported asset
