Improve how and when we display "Has Virtualized Data" in the content browser tooltip.

- If the VA system is enabled and the package is too old to have a package trailer then the tooltip should tell the user this and suggest that the package be re-saved. This should help people understand why attempting to virtualize an older package failed.
- If the VA system is disabled there is no reason to display the "Has Virtualized Data" tool tip line at all, this should help prevent users from worrying that they have accidently enabled the system.
-- However if a package in the project is virtualized for what ever reason we should still show that the package has virtualized data because this is likely to cause an error in the editor so we want the user to be able to see why.

diff --git a/Engine/Plugins/Editor/ContentBrowser/ContentBrowserAssetDataSource/Source/ContentBrowserAssetDataSource/Private/ContentBrowserAssetDataCore.cpp b/Engine/Plugins/Editor/ContentBrowser/ContentBrowserAssetDataSource/Source/ContentBrowserAssetDataSource/Private/ContentBrowserAssetDataCore.cpp
--- a/Engine/Plugins/Editor/ContentBrowser/ContentBrowserAssetDataSource/Source/ContentBrowserAssetDataSource/Private/ContentBrowserAssetDataCore.cpp
+++ b/Engine/Plugins/Editor/ContentBrowser/ContentBrowserAssetDataSource/Source/ContentBrowserAssetDataSource/Private/ContentBrowserAssetDataCore.cpp
@@ -3,30 +3,31 @@
 #include "ContentBrowserAssetDataCore.h"
 
 #include "AssetDefinition.h"
-#include "ContentBrowserDataSource.h"
-#include "ContentBrowserAssetDataPayload.h"
-#include "IAssetTools.h"
-#include "ContentBrowserDataSubsystem.h"
-#include "IContentBrowserDataModule.h"
+#include "AssetFileContextMenu.h"
+#include "AssetFolderContextMenu.h"
+#include "AssetPropertyTagCache.h"
 #include "AssetRegistry/IAssetRegistry.h"
 #include "AssetViewUtils.h"
-#include "AssetPropertyTagCache.h"
+#include "ContentBrowserAssetDataPayload.h"
+#include "ContentBrowserDataMenuContexts.h"
+#include "ContentBrowserDataSource.h"
+#include "ContentBrowserDataSubsystem.h"
+#include "ContentBrowserDataUtils.h"
+#include "Editor.h"
 #include "Editor/UnrealEdTypes.h"
-#include "ObjectTools.h"
-#include "HAL/FileManager.h"
+#include "Engine/Level.h"
 #include "FileHelpers.h"
-#include "Editor.h"
+#include "HAL/FileManager.h"
+#include "IAssetTools.h"
 #include "IAssetTypeActions.h"
-#include "Subsystems/AssetEditorSubsystem.h"
-#include "AssetFolderContextMenu.h"
-#include "AssetFileContextMenu.h"
-#include "ContentBrowserDataUtils.h"
-#include "ContentBrowserDataMenuContexts.h"
-#include "Engine/Level.h"
+#include "IContentBrowserDataModule.h"
 #include "Misc/PackageName.h"
-#include "Misc/WarnIfAssetsLoadedInScope.h"
 #include "Misc/Paths.h"
+#include "Misc/WarnIfAssetsLoadedInScope.h"
+#include "ObjectTools.h"
+#include "Subsystems/AssetEditorSubsystem.h"
 #include "ToolMenu.h"
+#include "Virtualization/VirtualizationSystem.h"
 
 #define LOCTEXT_NAMESPACE "ContentBrowserAssetDataSource"
 
@@ -1449,10 +1450,29 @@ bool GetVirtualizationItemAttribute(const FAssetData& InAssetData, IAssetRegistr
 
 	if (TOptional<FAssetPackageData> PackageData = InAssetRegistry->GetAssetPackageDataCopy(InAssetData.PackageName))
 	{
-		// We could set a bool here but that will display the value in lower case, where as asset properties 
-		// use a string value with the first letter in upper case, so we replicate that here to avoid the 
-		// entry looking out of place.
-		OutAttributeValue.SetValue(PackageData->HasVirtualizedPayloads() ? TEXT("True") : TEXT("False"));
+		if (UE::Virtualization::IVirtualizationSystem::Get().IsEnabled())
+		{
+			if (PackageData->FileVersionUE >= EUnrealEngineObjectUE5Version::PAYLOAD_TOC)
+			{
+				OutAttributeValue.SetValue(PackageData->HasVirtualizedPayloads() ? TEXT("True") : TEXT("False"));
+			}
+			else
+			{
+				OutAttributeValue.SetValue(TEXT("Version too old, resave to enable"));
+			}
+		}
+		else
+		{
+			if (PackageData->HasVirtualizedPayloads())
+			{
+				OutAttributeValue.SetValue(TEXT("True"));
+			}
+			else
+			{
+				// The VA system is disabled and nothing in the package is virtualized, so no need to display this in the tool tip
+				return false;
+			}
+		}
 
 		if (InIncludeMetaData)
 		{
