Localize the virtualization tooltip message for older packages.

diff --git a/Engine/Plugins/Editor/ContentBrowser/ContentBrowserAssetDataSource/Source/ContentBrowserAssetDataSource/Private/ContentBrowserAssetDataCore.cpp b/Engine/Plugins/Editor/ContentBrowser/ContentBrowserAssetDataSource/Source/ContentBrowserAssetDataSource/Private/ContentBrowserAssetDataCore.cpp
--- a/Engine/Plugins/Editor/ContentBrowser/ContentBrowserAssetDataSource/Source/ContentBrowserAssetDataSource/Private/ContentBrowserAssetDataCore.cpp
+++ b/Engine/Plugins/Editor/ContentBrowser/ContentBrowserAssetDataSource/Source/ContentBrowserAssetDataSource/Private/ContentBrowserAssetDataCore.cpp
@@ -1450,6 +1450,9 @@ bool GetVirtualizationItemAttribute(const FAssetData& InAssetData, IAssetRegistr
 
 	if (TOptional<FAssetPackageData> PackageData = InAssetRegistry->GetAssetPackageDataCopy(InAssetData.PackageName))
 	{
+		// Note that we do not localize True/False as nowhere else in the code base seems to do this
+		// and it would be odd for only a single entry in the tool tip to localize booleans.
+
 		if (UE::Virtualization::IVirtualizationSystem::Get().IsEnabled())
 		{
 			if (PackageData->FileVersionUE >= EUnrealEngineObjectUE5Version::PAYLOAD_TOC)
@@ -1458,7 +1461,7 @@ bool GetVirtualizationItemAttribute(const FAssetData& InAssetData, IAssetRegistr
 			}
 			else
 			{
-				OutAttributeValue.SetValue(TEXT("Version too old, resave to enable"));
+				OutAttributeValue.SetValue(LOCTEXT("AttributeDisplayName_OutOfDate", "Version too old, resave to enable"));
 			}
 		}
 		else
