Fix the blocked import of a folder when using the extension permissions.
Also I moved the permissions logic closer to the code that actually handle the import.

diff --git a/Engine/Plugins/Editor/ContentBrowser/ContentBrowserAssetDataSource/Source/ContentBrowserAssetDataSource/Private/ContentBrowserAssetDataSource.cpp b/Engine/Plugins/Editor/ContentBrowser/ContentBrowserAssetDataSource/Source/ContentBrowserAssetDataSource/Private/ContentBrowserAssetDataSource.cpp
--- a/Engine/Plugins/Editor/ContentBrowser/ContentBrowserAssetDataSource/Source/ContentBrowserAssetDataSource/Private/ContentBrowserAssetDataSource.cpp
+++ b/Engine/Plugins/Editor/ContentBrowser/ContentBrowserAssetDataSource/Source/ContentBrowserAssetDataSource/Private/ContentBrowserAssetDataSource.cpp
@@ -18,6 +18,7 @@
 #include "IAssetTools.h"
 #include "ToolMenus.h"
 #include "Misc/PackageName.h"
+#include "Misc/PathViews.h"
 #include "NewAssetContextMenu.h"
 #include "AssetFolderContextMenu.h"
 #include "AssetFileContextMenu.h"
@@ -2534,7 +2535,8 @@ bool UContentBrowserAssetDataSource::CanHandleDragDropEvent(const FContentBrowse
 				bool bSupportOneFile = false;
 				for (const FString& File : ExternalDragDropOp->GetFiles())
 				{
-					if (AssetTools->IsImportExtensionAllowed(FPaths::GetExtension(File)))
+					FStringView Extension = FPathViews::GetExtension(File);
+					if (Extension.IsEmpty() || AssetTools->IsImportExtensionAllowed(Extension))
 					{
 						bSupportOneFile = true;
 					}
@@ -2578,39 +2580,15 @@ bool UContentBrowserAssetDataSource::HandleDragDropOnItem(const FContentBrowserI
 			FString UnsupportedFiles;
 			if (ExternalDragDropOp->HasFiles() && ContentBrowserAssetData::CanModifyPath(AssetTools, FolderPayload->GetInternalPath(), &ErrorMsg))
 			{
-				TArray<FString> ImportFiles;
-				for (const FString& File : ExternalDragDropOp->GetFiles())
-				{
-					if (AssetTools->IsImportExtensionAllowed(FPaths::GetExtension(File)))
-					{
-						ImportFiles.AddUnique(File);
-					}
-					else
-					{
-						if (!UnsupportedFiles.IsEmpty())
-						{
-							UnsupportedFiles += TEXT("\n");
-						}
-						UnsupportedFiles += File;
-					}
-				}
-				if (ImportFiles.Num() > 0)
+				TArray<FString> ImportFiles = ExternalDragDropOp->GetFiles();
+			
+				if (!ImportFiles.IsEmpty())
 				{
 					// Delay import until next tick to avoid blocking the process that files were dragged from
 					GEditor->GetEditorSubsystem<UImportSubsystem>()->ImportNextTick(ImportFiles, FolderPayload->GetInternalPath().ToString());
 				}
 			}
 
-			if (!UnsupportedFiles.IsEmpty())
-			{
-				ErrorMsg = FText::Format(LOCTEXT("HandleDragDropOnItemUnsupportedFile", "Unsupported file format to import:\n\t{0}"), FText::FromString(UnsupportedFiles));
-			}
-
-			if (!ErrorMsg.IsEmpty())
-			{
-				AssetViewUtils::ShowErrorNotifcation(ErrorMsg);
-			}
-
 			return true; // We handled this drop, even if the result was invalid (eg, read-only folder)
 		}
 	}
diff --git a/Engine/Source/Developer/AssetTools/Private/AssetTools.cpp b/Engine/Source/Developer/AssetTools/Private/AssetTools.cpp
--- a/Engine/Source/Developer/AssetTools/Private/AssetTools.cpp
+++ b/Engine/Source/Developer/AssetTools/Private/AssetTools.cpp
@@ -96,6 +96,7 @@
 #include "UObject/LinkerInstancingContext.h"
 #include "PackageMigrationContext.h"
 #include "ComponentReregisterContext.h"
+#include "AssetViewUtils.h"
 
 #if WITH_EDITOR
 #include "Subsystems/AssetEditorSubsystem.h"
@@ -3431,6 +3432,38 @@ TArray<UObject*> UAssetToolsImpl::ImportAssetsInternal(const TArray<FString>& Fi
 		FilesAndDestinations = (*FilesAndDestinationsPtr);
 	}
 
+	{
+		FString UnallowedFilesString;
+		// Check for non allowed extensions
+		FilesAndDestinations.RemoveAll([&UnallowedFilesString, this](const TPair<FString, FString>& InValue)
+			{
+				FStringView Extension = FPathViews::GetExtension(InValue.Key);
+				if (!IsImportExtensionAllowed(Extension))
+				{
+					if (!UnallowedFilesString.IsEmpty())
+					{
+						UnallowedFilesString += TEXT("\n");
+					}
+					UnallowedFilesString += InValue.Key;
+
+					return true;
+				}
+
+				return false;
+			});
+
+		FText ErrorMsg;
+		if (!UnallowedFilesString.IsEmpty())
+		{
+			ErrorMsg = FText::Format(LOCTEXT("UnsupportedFileImport", "Unsupported file format to import:\n{0}"), FText::FromString(UnallowedFilesString));
+		}
+
+		if (!ErrorMsg.IsEmpty())
+		{
+			AssetViewUtils::ShowErrorNotifcation(ErrorMsg);
+		}
+	}
+
 	if(SpecifiedFactory == nullptr)
 	{
 		// First instantiate one factory for each file extension encountered that supports the extension
@@ -6163,10 +6196,10 @@ const TSharedRef<FNamePermissionList>& UAssetToolsImpl::GetImportExtensionPermis
 	return ImportExtensionPermissionList;
 }
 
-bool UAssetToolsImpl::IsImportExtensionAllowed(const FString& Extension) const
+bool UAssetToolsImpl::IsImportExtensionAllowed(const FStringView& Extension) const
 {
 	//Always test the extension in lower case
-	FName NameExtension(*Extension);
+	FName NameExtension(Extension);
 	return ImportExtensionPermissionList->PassesFilter(NameExtension);
 }
 
diff --git a/Engine/Source/Developer/AssetTools/Private/AssetTools.h b/Engine/Source/Developer/AssetTools/Private/AssetTools.h
--- a/Engine/Source/Developer/AssetTools/Private/AssetTools.h
+++ b/Engine/Source/Developer/AssetTools/Private/AssetTools.h
@@ -139,7 +139,7 @@ public:
 	virtual TArray<UFactory*> GetNewAssetFactories() const override;
 	virtual const TSharedRef<FPathPermissionList>& GetAssetClassPathPermissionList(EAssetClassAction AssetClassAction) const override;
 	virtual const TSharedRef<FNamePermissionList>& GetImportExtensionPermissionList() const override;
-	virtual bool IsImportExtensionAllowed(const FString& Extension) const override;
+	virtual bool IsImportExtensionAllowed(const FStringView& Extension) const override;
 	virtual TSet<EBlueprintType>& GetAllowedBlueprintTypes() override;
 	virtual TSharedRef<FPathPermissionList>& GetFolderPermissionList() override;
 	virtual TSharedRef<FPathPermissionList>& GetWritableFolderPermissionList() override;
diff --git a/Engine/Source/Developer/AssetTools/Public/IAssetTools.h b/Engine/Source/Developer/AssetTools/Public/IAssetTools.h
--- a/Engine/Source/Developer/AssetTools/Public/IAssetTools.h
+++ b/Engine/Source/Developer/AssetTools/Public/IAssetTools.h
@@ -595,7 +595,7 @@ public:
 	/** Get extension permission list allowed for importer */
 	virtual const TSharedRef<FNamePermissionList>& GetImportExtensionPermissionList() const = 0;
 
-	virtual bool IsImportExtensionAllowed(const FString& Extension) const = 0;
+	virtual bool IsImportExtensionAllowed(const FStringView& Extension) const = 0;
 
 	/** Which BlueprintTypes are allowed to be created. An empty list should allow everything. */
 	virtual TSet<EBlueprintType>& GetAllowedBlueprintTypes() = 0;
diff --git a/Engine/Source/Editor/UnrealEd/Private/Subsystems/ImportSubsystem.cpp b/Engine/Source/Editor/UnrealEd/Private/Subsystems/ImportSubsystem.cpp
--- a/Engine/Source/Editor/UnrealEd/Private/Subsystems/ImportSubsystem.cpp
+++ b/Engine/Source/Editor/UnrealEd/Private/Subsystems/ImportSubsystem.cpp
@@ -154,33 +154,8 @@ void UImportSubsystem::ImportNextTick(const TArray<FString>& Files, const FStrin
 		AssetToolsModule.Get().NotifyBlockedByWritableFolderFilter();
 		return;
 	}
-	FString UnsupportedFiles;
-	TArray<FString> ImportFiles;
-	for (const FString& File : Files)
-	{
-		if (AssetToolsModule.Get().IsImportExtensionAllowed(FPaths::GetExtension(File)))
-		{
-			ImportFiles.AddUnique(File);
-		}
-		else
-		{
-			if (!UnsupportedFiles.IsEmpty())
-			{
-				UnsupportedFiles += TEXT("\n");
-			}
-			UnsupportedFiles += File;
-		}
-	}
-	if (ImportFiles.Num() > 0)
-	{
-		PendingTasks.Enqueue(MakeShared<FImportFilesByPath>(ImportFiles, DestinationPath));
-	}
-
-	if (!UnsupportedFiles.IsEmpty())
-	{
-		FText ErrorMsg = FText::Format(NSLOCTEXT("UImportSubsystem", "ImportNextTickUnsupportedExtension", "Unsupported file format to import:\n\t{0}"), FText::FromString(UnsupportedFiles));
-		FSlateNotificationManager::Get().AddNotification(FNotificationInfo(ErrorMsg));
-	}
+	
+	PendingTasks.Enqueue(MakeShared<FImportFilesByPath>(Files, DestinationPath));
 
 	GEditor->GetTimerManager()->SetTimerForNextTick(this, &UImportSubsystem::HandleNextTick);
 }
