Make "Show Redirectors" and "Show Empty Folders" cooperate
- Move redirector filtering into backend
- Add a custom filter into SFilterBar to maintain old UI for redirector visibilty

diff --git a/Engine/Plugins/Editor/ContentBrowser/ContentBrowserAssetDataSource/Source/ContentBrowserAssetDataSource/Private/ContentBrowserAssetDataSource.cpp b/Engine/Plugins/Editor/ContentBrowser/ContentBrowserAssetDataSource/Source/ContentBrowserAssetDataSource/Private/ContentBrowserAssetDataSource.cpp
--- a/Engine/Plugins/Editor/ContentBrowser/ContentBrowserAssetDataSource/Source/ContentBrowserAssetDataSource/Private/ContentBrowserAssetDataSource.cpp
+++ b/Engine/Plugins/Editor/ContentBrowser/ContentBrowserAssetDataSource/Source/ContentBrowserAssetDataSource/Private/ContentBrowserAssetDataSource.cpp
@@ -41,6 +41,84 @@
 
 #define LOCTEXT_NAMESPACE "ContentBrowserAssetDataSource"
 
+enum class EContentBrowserFolderAttributes : uint8
+{
+	/**
+	 * No special attributes.
+	 */
+	None = 0,
+
+	/**
+	 * This folder should always be visible, even if it contains no content in the Content Browser view.
+	 * This will include root content folders, and any folders that have been created directly (or indirectly) by a user action.
+	 */
+	AlwaysVisible = 1 << 0,
+
+	/**
+	 * This folder has non-redirector assets that will appear in the Content Browser view.
+	 */
+	HasAssets = 1 << 1,
+
+	/**
+	 * This folder has public content that will appear in the Content Browser view.
+	 */
+	HasPublicContent = 1 << 2,
+
+	/**
+	 * This folder has source (uncooked) content that will appear in the Content Browser view.
+	 */
+	HasSourceContent = 1 << 3,
+
+	/**
+	 * This folder is inside a plugin.
+	 */
+	IsInPlugin = 1 << 4,
+
+	/**
+	 * This folder has redirector assets that will appear in the Content Browser view if the UI wishes to display them
+	 */
+	HasRedirectors = 1 << 5,
+};
+ENUM_CLASS_FLAGS(EContentBrowserFolderAttributes);
+
+// Produce a string of flags |'d together for logging
+FStringBuilderBase& operator<<(FStringBuilderBase& Builder, EContentBrowserFolderAttributes Attribs)
+{
+	bool bFirst = true;
+	for (EContentBrowserFolderAttributes Flag : MakeFlagsRange(Attribs))
+	{
+		if (!bFirst)
+		{
+			Builder << TEXTVIEW("|");
+		}
+		switch (Flag)
+		{
+			case EContentBrowserFolderAttributes::AlwaysVisible:
+				Builder << TEXTVIEW("AlwaysVisible");
+				break;
+			case EContentBrowserFolderAttributes::HasAssets:
+				Builder << TEXTVIEW("HasAssets");
+				break;
+			case EContentBrowserFolderAttributes::HasPublicContent:
+				Builder << TEXTVIEW("HasPublicContent");
+				break;
+			case EContentBrowserFolderAttributes::HasSourceContent:
+				Builder << TEXTVIEW("HasSourceContent");
+				break;
+			case EContentBrowserFolderAttributes::IsInPlugin:
+				Builder << TEXTVIEW("IsInPlugin");
+				break;
+			case EContentBrowserFolderAttributes::HasRedirectors:
+				Builder << TEXTVIEW("HasRedirectors");
+				break;
+			default:
+				Builder << TEXTVIEW("Unknown");
+				break;
+		}
+	}
+	return Builder;
+}
+
 UContentBrowserAssetDataSource::FOnAssetDataSourcePathAdded UContentBrowserAssetDataSource::OnAssetPathAddedDelegate;
 UContentBrowserAssetDataSource::FOnAssetDataSourcePathRemoved UContentBrowserAssetDataSource::OnAssetPathRemovedDelegate;
 
@@ -219,9 +297,16 @@ bool UContentBrowserAssetDataSource::PopulateAssetFilterInputParams(FAssetFilter
 	Params.bIncludeFolders = EnumHasAnyFlags(InFilter.ItemTypeFilter, EContentBrowserItemTypeFilter::IncludeFolders);
 	Params.bIncludeFiles = EnumHasAnyFlags(InFilter.ItemTypeFilter, EContentBrowserItemTypeFilter::IncludeFiles);
 	Params.bIncludeAssets = EnumHasAnyFlags(InFilter.ItemCategoryFilter, EContentBrowserItemCategoryFilter::IncludeAssets);
+	Params.bIncludeRedirectors = EnumHasAnyFlags(InFilter.ItemCategoryFilter, EContentBrowserItemCategoryFilter::IncludeRedirectors);
 
-	// If we aren't including anything, then we can just bail now
-	if (!Params.bIncludeAssets || (!Params.bIncludeFolders && !Params.bIncludeFiles))
+	// Everything this data source tracks is either an asset or a redirector
+	if (!Params.bIncludeAssets && !Params.bIncludeRedirectors)
+	{
+		return false;
+	}
+
+	// Everything this data source tracks is either a file or a folder
+	if (!Params.bIncludeFolders && !Params.bIncludeFiles)
 	{
 		return false;
 	}
@@ -252,6 +337,7 @@ bool UContentBrowserAssetDataSource::PopulateAssetFilterInputParams(FAssetFilter
 	Params.AssetDataFilter = &Params.FilterList->FindOrAddFilter<FContentBrowserCompiledAssetDataFilter>();
 	Params.AssetDataFilter->bFilterExcludesAllAssets = true;
 	Params.AssetDataFilter->ItemAttributeFilter = InFilter.ItemAttributeFilter;
+	Params.AssetDataFilter->ItemCategoryFilter = InFilter.ItemCategoryFilter;
 	Params.InternalPaths.Reset();
 
 	Params.UnsupportedClassFilter = InFilter.ExtraFilters.FindFilter<FContentBrowserDataUnsupportedClassFilter>();
@@ -787,6 +873,16 @@ PRAGMA_ENABLE_DEPRECATION_WARNINGS
 				ExclusiveFilter.bRecursiveClasses |= Params.ClassFilter->bRecursiveClassNamesToExclude;
 			}
 
+			if (!Params.bIncludeRedirectors)
+			{
+				// Add redirectors to the class paths to exclude unless it's already in the class paths to explicitly include
+				FTopLevelAssetPath RedirectorClassPath(UObjectRedirector::StaticClass());
+				if (!CompiledInclusiveFilter.ClassPaths.Contains(RedirectorClassPath))
+				{
+					ExclusiveFilter.ClassPaths.Add(RedirectorClassPath);
+				}
+			}
+
 			Params.AssetRegistry->CompileFilter(ExclusiveFilter, CompiledExclusiveFilter);
 		}
 
@@ -1058,6 +1154,7 @@ PRAGMA_ENABLE_DEPRECATION_WARNINGS
 	return true;
 }
 
+// Note that this function is deprecated and is no longer maintained, see declaration
 bool UContentBrowserAssetDataSource::CreateAssetFilter(FAssetFilterInputParams& Params, FName InPath, const FContentBrowserDataFilter& InFilter, FContentBrowserDataCompiledFilter& OutCompiledFilter, FCompileARFilterFunc CreateCompiledFilter)
 {
 	TRACE_CPUPROFILER_EVENT_SCOPE(UContentBrowserAssetDataSource::LegacyCreateAssetFilter);
@@ -1932,8 +2029,40 @@ bool UContentBrowserAssetDataSource::PrioritizeSearchPath(const FName InPath)
 	return true;
 }
 
-bool UContentBrowserAssetDataSource::IsFolderVisible(const FName InPath, const EContentBrowserIsFolderVisibleFlags InFlags)
+bool UContentBrowserAssetDataSource::IsFolderVisible(const FName InPath, const EContentBrowserIsFolderVisibleFlags InFlags, TOptional<FContentBrowserFolderContentsFilter> InContentsFilter)
 {
+	auto IsInternalFolderVisible = [this, InContentsFilter](FName InternalFolderPath) {
+		const EContentBrowserFolderAttributes FolderAttributes = GetAssetFolderAttributes(InternalFolderPath);
+		if (EnumHasAnyFlags(FolderAttributes, EContentBrowserFolderAttributes::AlwaysVisible))
+		{
+			return true;
+		}
+
+		// Hide folders that only contain cooked private content
+		if (EnumHasAnyFlags(FolderAttributes, EContentBrowserFolderAttributes::HasAssets | EContentBrowserFolderAttributes::HasRedirectors)
+			&& !EnumHasAnyFlags(FolderAttributes, EContentBrowserFolderAttributes::HasPublicContent | EContentBrowserFolderAttributes::HasSourceContent))
+		{
+			return false;
+		}
+
+		if (InContentsFilter.IsSet())
+		{
+			if (EnumHasAnyFlags(InContentsFilter->ItemCategoryFilter, EContentBrowserItemCategoryFilter::IncludeAssets)
+				&& EnumHasAnyFlags(FolderAttributes, EContentBrowserFolderAttributes::HasAssets))
+			{
+				return true;
+			}
+			if (EnumHasAnyFlags(InContentsFilter->ItemCategoryFilter, EContentBrowserItemCategoryFilter::IncludeRedirectors)
+				&& EnumHasAnyFlags(FolderAttributes, EContentBrowserFolderAttributes::HasRedirectors))
+			{
+				return true;
+			}
+			return false;
+		}
+
+		return true;
+	};
+
 	FName ConvertedPath;
 	const EContentBrowserPathType ConvertedPathType = TryConvertVirtualPath(InPath, ConvertedPath);
 	if (ConvertedPathType == EContentBrowserPathType::Internal)
@@ -1942,30 +2071,28 @@ bool UContentBrowserAssetDataSource::IsFolderVisible(const FName InPath, const E
 		{
 			return false;
 		}
+		return IsInternalFolderVisible(ConvertedPath);
 	}
 	else if (ConvertedPathType == EContentBrowserPathType::Virtual)
 	{
+		bool bAnyVisible = false;
+		// Make virtual folders visible if any of their child folders will be visible
+		RootPathVirtualTree.EnumerateSubPaths(
+			ConvertedPath,
+			[this, &bAnyVisible, &IsInternalFolderVisible](FName ChildVirtualPath, FName ChildInternalPath) -> bool {
+				if (!ChildInternalPath.IsNone())
+				{
+					bAnyVisible = IsInternalFolderVisible(ChildInternalPath);
+				}
+				return !bAnyVisible;
+			},
+			true);
 		return true;
 	}
 	else
 	{
 		return false;
 	}
-
-	const EContentBrowserFolderAttributes FolderAttributes = GetAssetFolderAttributes(ConvertedPath);
-	if (EnumHasAnyFlags(FolderAttributes, EContentBrowserFolderAttributes::AlwaysVisible))
-	{
-		return true;
-	}
-
-	// Hide folders that only contain cooked private content
-	if (EnumHasAnyFlags(FolderAttributes, EContentBrowserFolderAttributes::HasContent) && !EnumHasAnyFlags(FolderAttributes, EContentBrowserFolderAttributes::HasPublicContent | EContentBrowserFolderAttributes::HasSourceContent))
-	{
-		return false;
-	}
-
-	return !EnumHasAnyFlags(InFlags, EContentBrowserIsFolderVisibleFlags::HideEmptyFolders)
-		|| EnumHasAnyFlags(FolderAttributes, EContentBrowserFolderAttributes::HasContent);
 }
 
 bool UContentBrowserAssetDataSource::CanCreateFolder(const FName InPath, FText* OutErrorMsg)
@@ -2701,7 +2828,8 @@ FContentBrowserItemData UContentBrowserAssetDataSource::CreateAssetFolderItem(co
 	TryConvertInternalPathToVirtual(InFolderPath, VirtualizedPath);
 
 	const EContentBrowserFolderAttributes FolderAttributes = GetAssetFolderAttributes(InFolderPath);
-	const bool bIsCookedPath = EnumHasAnyFlags(FolderAttributes, EContentBrowserFolderAttributes::HasContent) && !EnumHasAnyFlags(FolderAttributes, EContentBrowserFolderAttributes::HasSourceContent);
+	const bool bIsCookedPath = EnumHasAnyFlags(FolderAttributes, EContentBrowserFolderAttributes::HasAssets | EContentBrowserFolderAttributes::HasRedirectors)
+							&& !EnumHasAnyFlags(FolderAttributes, EContentBrowserFolderAttributes::HasSourceContent);
 	const bool bIsPlugin = EnumHasAnyFlags(FolderAttributes, EContentBrowserFolderAttributes::IsInPlugin);
 	return ContentBrowserAssetData::CreateAssetFolderItem(this, VirtualizedPath, InFolderPath, bIsCookedPath, bIsPlugin);
 }
@@ -2894,10 +3022,10 @@ void UContentBrowserAssetDataSource::OnPathsRemoved(TConstArrayView<FStringView>
 
 void UContentBrowserAssetDataSource::OnPathPopulated(const FAssetData& InAssetData)
 {
-	const EContentBrowserFolderAttributes FolderAttributes 
-		= EContentBrowserFolderAttributes::HasContent
-		| (InAssetData.PackageFlags & PKG_Cooked ? EContentBrowserFolderAttributes::None : EContentBrowserFolderAttributes::HasSourceContent)
-		| (InAssetData.PackageFlags & PKG_NotExternallyReferenceable ? EContentBrowserFolderAttributes::None : EContentBrowserFolderAttributes::HasPublicContent);
+	EContentBrowserFolderAttributes FolderAttributes =
+		InAssetData.IsRedirector() ? EContentBrowserFolderAttributes::HasRedirectors : EContentBrowserFolderAttributes::HasAssets;
+	FolderAttributes |= (InAssetData.PackageFlags & PKG_Cooked ? EContentBrowserFolderAttributes::None : EContentBrowserFolderAttributes::HasSourceContent);
+	FolderAttributes |= (InAssetData.PackageFlags & PKG_NotExternallyReferenceable ? EContentBrowserFolderAttributes::None : EContentBrowserFolderAttributes::HasPublicContent);
 
 	OnPathPopulated(FNameBuilder(InAssetData.PackagePath), FolderAttributes);
 }
@@ -3017,6 +3145,11 @@ bool UContentBrowserAssetDataSource::SetAssetFolderAttributes(const FName InPath
 		EnumAddFlags(FolderAttributes, InAttributesToSet);
 
 		const bool bHasChanged = FolderAttributes != PreviousAttributes;
+		if (bHasChanged)
+		{
+			const EContentBrowserFolderAttributes NewAttributes = InAttributesToSet & ~(PreviousAttributes);
+			UE_LOG(LogContentBrowserAssetDataSource, Verbose, TEXT("Updated folder attributes: %s %s"), *WriteToString<256>(InPath), *WriteToString<256>(NewAttributes));
+		}
 		return bHasChanged;
 	}
 
@@ -3385,9 +3518,14 @@ FContentBrowserItemData UContentBrowserAssetDataSource::OnFinalizeDuplicateAsset
 	return CreateAssetFileItem(FAssetData(Asset));
 }
 
-bool UContentBrowserAssetDataSource::PathPassesCompiledDataFilter(const FContentBrowserCompiledAssetDataFilter& InFilter, const FName InPath)
+bool UContentBrowserAssetDataSource::PathPassesCompiledDataFilter(const FContentBrowserCompiledAssetDataFilter& InFilter, const FName InInternalPath)
 {
-	FNameBuilder PathStr(InPath);
+	if (InFilter.ExcludedPackagePaths.Contains(InInternalPath)) // PassesExcludedPathsFilter
+	{
+		return false;
+	}
+
+	FNameBuilder PathStr(InInternalPath);
 	FStringView Path(PathStr);
 
 	auto PathPassesFilter = [Path](const FPathPermissionList& InPathFilter, const bool InRecursive)
@@ -3397,8 +3535,7 @@ bool UContentBrowserAssetDataSource::PathPassesCompiledDataFilter(const FContent
 
 	return PathPassesFilter(InFilter.PackagePathsToInclude, InFilter.bRecursivePackagePathsToInclude)
 		&& PathPassesFilter(InFilter.PackagePathsToExclude, InFilter.bRecursivePackagePathsToExclude)
-		&& PathPassesFilter(InFilter.PathPermissionList, /*bRecursive*/true) // PassesPathFilter
-		&& !InFilter.ExcludedPackagePaths.Contains(InPath) // PassesExcludedPathsFilter
+		&& PathPassesFilter(InFilter.PathPermissionList, /*bRecursive*/ true)                         // PassesPathFilter
 		&& ContentBrowserDataUtils::PathPassesAttributeFilter(Path, 0, InFilter.ItemAttributeFilter); // PassesAttributeFilter
 }
 
diff --git a/Engine/Plugins/Editor/ContentBrowser/ContentBrowserAssetDataSource/Source/ContentBrowserAssetDataSource/Public/ContentBrowserAssetDataSource.h b/Engine/Plugins/Editor/ContentBrowser/ContentBrowserAssetDataSource/Source/ContentBrowserAssetDataSource/Public/ContentBrowserAssetDataSource.h
--- a/Engine/Plugins/Editor/ContentBrowser/ContentBrowserAssetDataSource/Source/ContentBrowserAssetDataSource/Public/ContentBrowserAssetDataSource.h
+++ b/Engine/Plugins/Editor/ContentBrowser/ContentBrowserAssetDataSource/Source/ContentBrowserAssetDataSource/Public/ContentBrowserAssetDataSource.h
@@ -44,6 +44,7 @@ public:
 	FPathPermissionList PathPermissionList;
 	TSet<FName> ExcludedPackagePaths;
 	EContentBrowserItemAttributeFilter ItemAttributeFilter = EContentBrowserItemAttributeFilter::IncludeAll;
+	EContentBrowserItemCategoryFilter ItemCategoryFilter = EContentBrowserItemCategoryFilter::IncludeAll;
 	FString VirtualPathToScanOnDemand;
 	// Cached filtering
 	TSet<FName> CachedSubPaths;
@@ -76,40 +77,7 @@ public:
 	FARCompiledFilter ExclusiveFilter;
 };
 
-enum class EContentBrowserFolderAttributes : uint8
-{
-	/**
-	 * No special attributes.
-	 */
-	None = 0,
-
-	/**
-	 * This folder should always be visible, even if it contains no content in the Content Browser view.
-	 * This will include root content folders, and any folders that have been created directly (or indirectly) by a user action.
-	 */
-	AlwaysVisible = 1<<0,
-
-	/**
-	 * This folder has content that will appear in the Content Browser view.
-	 */
-	HasContent = 1<<1,
-
-	/**
-	 * This folder has public content that will appear in the Content Browser view.
-	 */
-	HasPublicContent = 1<<2,
-
-	/**
-	 * This folder has source (uncooked) content that will appear in the Content Browser view.
-	 */
-	HasSourceContent = 1<<3,
-
-	/** 
-	 * This folder is inside a plugin.
-	 */
-	IsInPlugin = 1<<4,
-};
-ENUM_CLASS_FLAGS(EContentBrowserFolderAttributes);
+enum class EContentBrowserFolderAttributes : uint8;
 
 UCLASS()
 class CONTENTBROWSERASSETDATASOURCE_API UContentBrowserAssetDataSource : public UContentBrowserDataSource
@@ -192,6 +160,7 @@ public:
 		bool bIncludeFolders = false;
 		bool bIncludeFiles = false;
 		bool bIncludeAssets = false;
+		bool bIncludeRedirectors = false;
 	};
 
 	typedef TFunctionRef<void(FName, TFunctionRef<bool(FName)>, bool)> FSubPathEnumerationFunc;
@@ -279,7 +248,7 @@ public:
 
 	virtual bool PrioritizeSearchPath(const FName InPath) override;
 
-	virtual bool IsFolderVisible(const FName InPath, const EContentBrowserIsFolderVisibleFlags InFlags) override;
+	virtual bool IsFolderVisible(const FName InPath, const EContentBrowserIsFolderVisibleFlags InFlags, TOptional<FContentBrowserFolderContentsFilter> InContentsFilter) override;
 
 	virtual bool CanCreateFolder(const FName InPath, FText* OutErrorMsg) override;
 
diff --git a/Engine/Plugins/Editor/ContentBrowser/ContentBrowserClassDataSource/Source/ContentBrowserClassDataSource/Private/ContentBrowserClassDataSource.cpp b/Engine/Plugins/Editor/ContentBrowser/ContentBrowserClassDataSource/Source/ContentBrowserClassDataSource/Private/ContentBrowserClassDataSource.cpp
--- a/Engine/Plugins/Editor/ContentBrowser/ContentBrowserClassDataSource/Source/ContentBrowserClassDataSource/Private/ContentBrowserClassDataSource.cpp
+++ b/Engine/Plugins/Editor/ContentBrowser/ContentBrowserClassDataSource/Source/ContentBrowserClassDataSource/Private/ContentBrowserClassDataSource.cpp
@@ -416,8 +416,14 @@ bool UContentBrowserClassDataSource::EnumerateItemsForObjects(const TArrayView<U
 	return true;
 }
 
-bool UContentBrowserClassDataSource::IsFolderVisible(const FName InPath, const EContentBrowserIsFolderVisibleFlags InFlags)
+bool UContentBrowserClassDataSource::IsFolderVisible(const FName InPath, const EContentBrowserIsFolderVisibleFlags InFlags, TOptional<FContentBrowserFolderContentsFilter> InContentsFilter)
 {
+	// We only contain classes, bail if the caller wants to filter out folders that contain only classes
+	if (InContentsFilter.IsSet() && !EnumHasAnyFlags(InContentsFilter->ItemCategoryFilter, EContentBrowserItemCategoryFilter::IncludeClasses))
+	{
+		return false;
+	}
+
 	FName ConvertedPath;
 	const EContentBrowserPathType ConvertedPathType = TryConvertVirtualPath(InPath, ConvertedPath);
 	if (ConvertedPathType == EContentBrowserPathType::Internal)
@@ -438,9 +444,17 @@ bool UContentBrowserClassDataSource::IsFolderVisible(const FName InPath, const E
 
 	ConditionalCreateNativeClassHierarchy();
 
-	return ContentBrowserDataUtils::IsTopLevelFolder(ConvertedPath) 
-		|| !EnumHasAnyFlags(InFlags, EContentBrowserIsFolderVisibleFlags::HideEmptyFolders)
-		|| NativeClassHierarchy->HasClasses(ConvertedPath, /*bRecursive*/true);
+	if (ContentBrowserDataUtils::IsTopLevelFolder(ConvertedPath))
+	{
+		return true;
+	}
+
+	// Class flag was checked above - if we are filtering out folders that don't contain "class" elements, we are filtering out all empty folders from this provider
+	if (InContentsFilter.IsSet() && !NativeClassHierarchy->HasClasses(ConvertedPath, /*bRecursive*/ true))
+	{
+		return false;
+	}
+	return true;
 }
 
 bool UContentBrowserClassDataSource::DoesItemPassFilter(const FContentBrowserItemData& InItem, const FContentBrowserDataCompiledFilter& InFilter)
diff --git a/Engine/Plugins/Editor/ContentBrowser/ContentBrowserClassDataSource/Source/ContentBrowserClassDataSource/Public/ContentBrowserClassDataSource.h b/Engine/Plugins/Editor/ContentBrowser/ContentBrowserClassDataSource/Source/ContentBrowserClassDataSource/Public/ContentBrowserClassDataSource.h
--- a/Engine/Plugins/Editor/ContentBrowser/ContentBrowserClassDataSource/Source/ContentBrowserClassDataSource/Public/ContentBrowserClassDataSource.h
+++ b/Engine/Plugins/Editor/ContentBrowser/ContentBrowserClassDataSource/Source/ContentBrowserClassDataSource/Public/ContentBrowserClassDataSource.h
@@ -47,7 +47,7 @@ public:
 
 	virtual bool EnumerateItemsForObjects(const TArrayView<UObject*> InObjects, TFunctionRef<bool(FContentBrowserItemData&&)> InCallback) override;
 
-	virtual bool IsFolderVisible(const FName InPath, const EContentBrowserIsFolderVisibleFlags InFlags) override;
+	virtual bool IsFolderVisible(const FName InPath, const EContentBrowserIsFolderVisibleFlags InFlags, TOptional<FContentBrowserFolderContentsFilter> InContentsFilter) override;
 
 	virtual bool DoesItemPassFilter(const FContentBrowserItemData& InItem, const FContentBrowserDataCompiledFilter& InFilter) override;
 
diff --git a/Engine/Plugins/Editor/ContentBrowser/ContentBrowserFileDataSource/Source/ContentBrowserFileDataSource/Private/ContentBrowserFileDataSource.cpp b/Engine/Plugins/Editor/ContentBrowser/ContentBrowserFileDataSource/Source/ContentBrowserFileDataSource/Private/ContentBrowserFileDataSource.cpp
--- a/Engine/Plugins/Editor/ContentBrowser/ContentBrowserFileDataSource/Source/ContentBrowserFileDataSource/Private/ContentBrowserFileDataSource.cpp
+++ b/Engine/Plugins/Editor/ContentBrowser/ContentBrowserFileDataSource/Source/ContentBrowserFileDataSource/Private/ContentBrowserFileDataSource.cpp
@@ -1063,8 +1063,13 @@ bool UContentBrowserFileDataSource::PrioritizeSearchPath(const FName InPath)
 	return true;
 }
 
-bool UContentBrowserFileDataSource::IsFolderVisible(const FName InPath, const EContentBrowserIsFolderVisibleFlags InFlags)
+bool UContentBrowserFileDataSource::IsFolderVisible(const FName InPath, const EContentBrowserIsFolderVisibleFlags InFlags, TOptional<FContentBrowserFolderContentsFilter> ContentsFilter)
 {
+	if (ContentsFilter.IsSet() && !EnumHasAllFlags(ContentsFilter->ItemCategoryFilter, EContentBrowserItemCategoryFilter::IncludeMisc))
+	{
+		return false;
+	}
+
 	FName ConvertedPath;
 	const EContentBrowserPathType ConvertedPathType = TryConvertVirtualPath(InPath, ConvertedPath);
 	if (ConvertedPathType == EContentBrowserPathType::Internal)
@@ -1087,9 +1092,16 @@ bool UContentBrowserFileDataSource::IsFolderVisible(const FName InPath, const EC
 	const FStringView InternalPathStrView = InternalPathStr;
 	const uint32 InternalPathHash = GetTypeHash(InternalPathStrView);
 
-	return AlwaysVisibleFolders.ContainsByHash(InternalPathHash, InternalPathStrView)
-		|| !EnumHasAnyFlags(InFlags, EContentBrowserIsFolderVisibleFlags::HideEmptyFolders)
-		|| !EmptyFolders.ContainsByHash(InternalPathHash, InternalPathStrView);
+	if (AlwaysVisibleFolders.ContainsByHash(InternalPathHash, InternalPathStrView))
+	{
+		return true;
+	}
+	// Misc flag was checked above - if we are filtering out folders that don't contain "misc" elements, we are filtering out all empty folders from this provider
+	if (ContentsFilter.IsSet() && EmptyFolders.ContainsByHash(InternalPathHash, InternalPathStrView))
+	{
+		return false;
+	}
+	return true;
 }
 
 bool UContentBrowserFileDataSource::DoesItemPassFilter(const FContentBrowserItemData& InItem, const FContentBrowserDataCompiledFilter& InFilter)
diff --git a/Engine/Plugins/Editor/ContentBrowser/ContentBrowserFileDataSource/Source/ContentBrowserFileDataSource/Public/ContentBrowserFileDataSource.h b/Engine/Plugins/Editor/ContentBrowser/ContentBrowserFileDataSource/Source/ContentBrowserFileDataSource/Public/ContentBrowserFileDataSource.h
--- a/Engine/Plugins/Editor/ContentBrowser/ContentBrowserFileDataSource/Source/ContentBrowserFileDataSource/Public/ContentBrowserFileDataSource.h
+++ b/Engine/Plugins/Editor/ContentBrowser/ContentBrowserFileDataSource/Source/ContentBrowserFileDataSource/Public/ContentBrowserFileDataSource.h
@@ -59,7 +59,7 @@ public:
 
 	virtual bool PrioritizeSearchPath(const FName InPath) override;
 
-	virtual bool IsFolderVisible(const FName InPath, const EContentBrowserIsFolderVisibleFlags InFlags) override;
+	virtual bool IsFolderVisible(const FName InPath, const EContentBrowserIsFolderVisibleFlags InFlags, TOptional<FContentBrowserFolderContentsFilter> InContentsFilter) override;
 
 	virtual bool DoesItemPassFilter(const FContentBrowserItemData& InItem, const FContentBrowserDataCompiledFilter& InFilter) override;
 
diff --git a/Engine/Source/Editor/ContentBrowser/Private/ContentBrowserLog.h b/Engine/Source/Editor/ContentBrowser/Private/ContentBrowserLog.h
--- a/Engine/Source/Editor/ContentBrowser/Private/ContentBrowserLog.h
+++ b/Engine/Source/Editor/ContentBrowser/Private/ContentBrowserLog.h
@@ -4,4 +4,4 @@
 
 #include "CoreMinimal.h"
 
-DECLARE_LOG_CATEGORY_EXTERN(LogContentBrowser, Verbose, All);
+DECLARE_LOG_CATEGORY_EXTERN(LogContentBrowser, Log, Log);
diff --git a/Engine/Source/Editor/ContentBrowser/Private/ContentBrowserUtils.cpp b/Engine/Source/Editor/ContentBrowser/Private/ContentBrowserUtils.cpp
--- a/Engine/Source/Editor/ContentBrowser/Private/ContentBrowserUtils.cpp
+++ b/Engine/Source/Editor/ContentBrowser/Private/ContentBrowserUtils.cpp
@@ -18,6 +18,7 @@
 #include "Framework/Application/IMenu.h"
 #include "Framework/Application/MenuStack.h"
 #include "Framework/Application/SlateApplication.h"
+#include "FrontendFilterBase.h"
 #include "HAL/FileManager.h"
 #include "HAL/PlatformApplicationMisc.h"
 #include "HAL/PlatformCrt.h"
@@ -40,6 +41,7 @@
 #include "Misc/Optional.h"
 #include "Misc/Paths.h"
 #include "SAssetView.h"
+#include "SFilterList.h"
 #include "SPathView.h"
 #include "SlateOptMacros.h"
 #include "SlotBase.h"
@@ -785,7 +787,9 @@ FName ContentBrowserUtils::GetInvariantPath(const FContentBrowserItemPath& ItemP
 
 EContentBrowserIsFolderVisibleFlags ContentBrowserUtils::GetIsFolderVisibleFlags(const bool bDisplayEmpty)
 {
+	PRAGMA_DISABLE_DEPRECATION_WARNINGS
 	return EContentBrowserIsFolderVisibleFlags::Default | (bDisplayEmpty ? EContentBrowserIsFolderVisibleFlags::None : EContentBrowserIsFolderVisibleFlags::HideEmptyFolders);
+	PRAGMA_ENABLE_DEPRECATION_WARNINGS
 }
 
 bool ContentBrowserUtils::IsFavoriteFolder(const FString& FolderPath)
@@ -921,4 +925,17 @@ bool ContentBrowserUtils::ShouldShowPluginFolderIcon()
 	return CVarShowPluginFolderIcon->GetBool();
 }
 
+bool ContentBrowserUtils::ShouldShowRedirectors(TSharedPtr<SFilterList> Filters)
+{
+	if (Filters.IsValid())
+	{
+		TSharedPtr<FFrontendFilter> ShowRedirectorsFilter = Filters->GetFrontendFilter(TEXT("ShowRedirectorsBackend"));
+		if (ShowRedirectorsFilter.IsValid())
+		{
+			return Filters->IsFrontendFilterActive(ShowRedirectorsFilter);
+		}
+	}
+	return false;
+}
+
 #undef LOCTEXT_NAMESPACE
diff --git a/Engine/Source/Editor/ContentBrowser/Private/ContentBrowserUtils.h b/Engine/Source/Editor/ContentBrowser/Private/ContentBrowserUtils.h
--- a/Engine/Source/Editor/ContentBrowser/Private/ContentBrowserUtils.h
+++ b/Engine/Source/Editor/ContentBrowser/Private/ContentBrowserUtils.h
@@ -18,6 +18,7 @@
 class FPathPermissionList;
 class FSlateRect;
 class SAssetView;
+class SFilterList;
 class SPathView;
 class SWidget;
 struct FARFilter;
@@ -140,4 +141,7 @@ namespace ContentBrowserUtils
 
 	/** Returns whether we should display icons for plugins in the content browser */
 	bool ShouldShowPluginFolderIcon();
+	
+	/** Returns whether the content browser should be showing redirectors based on the current filter state or content browser settings */
+	bool ShouldShowRedirectors(TSharedPtr<SFilterList> Filters);
 }
diff --git a/Engine/Source/Editor/ContentBrowser/Private/Filters.cpp b/Engine/Source/Editor/ContentBrowser/Private/Filters.cpp
new file mode 100644
--- /dev/null
+++ b/Engine/Source/Editor/ContentBrowser/Private/Filters.cpp
@@ -0,0 +1,55 @@
+// Copyright Epic Games, Inc. All Rights Reserved.
+
+#include "Filters.h"
+
+#include "FrontendFilterBase.h"
+
+#define LOCTEXT_NAMESPACE "ContentBrowser"
+
+FFilter_ShowRedirectors::FFilter_ShowRedirectors(TSharedPtr<FFrontendFilterCategory> InCategory)
+	: FFrontendFilter(InCategory)
+{
+}
+
+/** Returns the human readable name for this filter */
+FText FFilter_ShowRedirectors::GetDisplayName() const
+{
+	return LOCTEXT("FrontendFilter_ShowRedirectors", "Show Redirectors");
+}
+
+/** Returns the tooltip for this filter, shown in the filters menu */
+FText FFilter_ShowRedirectors::GetToolTipText() const
+{
+	return LOCTEXT("FrontendFilter_ShowRedirectorsToolTip", "Allow display of Redirectors.");
+}
+
+/** Returns the name of the icon to use in menu entries */
+FName FFilter_ShowRedirectors::GetIconName() const
+{
+	return NAME_None;
+}
+
+/** Notification that the filter became active or inactive */
+void FFilter_ShowRedirectors::ActiveStateChanged(bool bActive)
+{
+	// Do nothing, filter state is queried externally e.g. by SContentBrowser
+}
+
+/** Can be overriden for custom FilterBar subclasses to save settings, currently not implemented in any gneeric Filter Bar */
+void FFilter_ShowRedirectors::SaveSettings(const FString& IniFilename, const FString& IniSection, const FString& SettingsString) const 
+{
+
+}
+
+/** Can be overriden for custom FilterBar subclasses to load settings, currently not implemented in any gneeric Filter Bar */
+void FFilter_ShowRedirectors::LoadSettings(const FString& IniFilename, const FString& IniSection, const FString& SettingsString) 
+{
+
+}
+
+bool FFilter_ShowRedirectors::PassesFilter(FAssetFilterType InItem) const
+{
+	return true; // All items pass, this filter just communicates with the backend 
+}
+
+#undef LOCTEXT_NAMESPACE
\ No newline at end of file
diff --git a/Engine/Source/Editor/ContentBrowser/Private/Filters.h b/Engine/Source/Editor/ContentBrowser/Private/Filters.h
new file mode 100644
--- /dev/null
+++ b/Engine/Source/Editor/ContentBrowser/Private/Filters.h
@@ -0,0 +1,54 @@
+// Copyright Epic Games, Inc. All Rights Reserved.
+
+#pragma once
+
+// Private filters not exposed to modules that depend on ContentBrowser
+
+#include "IContentBrowserSingleton.h"
+#include "Filters/GenericFilter.h"
+#include "FrontendFilterBase.h"
+#include "Templates/SharedPointer.h"
+
+/** 
+ *  A custom filter that appears in the filter bar but actually controls a content browser setting controlling 
+ *  visibility of redirectors for use in backend filtering.
+ */
+class FFilter_ShowRedirectors : public FFrontendFilter
+{
+public:
+	// TODO: Consider not giving this the entire viewmodel to sync with filter state
+	// could make viewmode own/create this, or send a smaller interface ptr into the constructor?
+	FFilter_ShowRedirectors (TSharedPtr<FFrontendFilterCategory> InCategory);
+
+	/** Returns the system name for this filter */
+	virtual FString GetName() const override { return TEXT("ShowRedirectorsBackend"); } 
+
+	/** Returns the human readable name for this filter */
+	virtual FText GetDisplayName() const override;
+
+	/** Returns the tooltip for this filter, shown in the filters menu */
+	virtual FText GetToolTipText() const override;
+
+	/** Returns the name of the icon to use in menu entries */
+	virtual FName GetIconName() const override;
+
+	/** If true, the filter will be active in the FilterBar when it is inactive in the UI (i.e the filter pill is grayed out)
+	 * @See: FFrontendFilter_ShowOtherDevelopers in Content Browser
+	 */
+	virtual bool IsInverseFilter() const 
+	{
+		// This has to be an inverse filter to prevent the asset view from recursively displaying all assets 
+		return true; 
+	}
+
+	/** Notification that the filter became active or inactive */
+	virtual void ActiveStateChanged(bool bActive) override;
+
+	/** Can be overriden for custom FilterBar subclasses to save settings, currently not implemented in any gneeric Filter Bar */
+	virtual void SaveSettings(const FString& IniFilename, const FString& IniSection, const FString& SettingsString) const override;
+
+	/** Can be overriden for custom FilterBar subclasses to load settings, currently not implemented in any gneeric Filter Bar */
+	virtual void LoadSettings(const FString& IniFilename, const FString& IniSection, const FString& SettingsString) override;
+
+	virtual bool PassesFilter(FAssetFilterType InItem) const override;
+};
\ No newline at end of file
diff --git a/Engine/Source/Editor/ContentBrowser/Private/SAssetPicker.cpp b/Engine/Source/Editor/ContentBrowser/Private/SAssetPicker.cpp
--- a/Engine/Source/Editor/ContentBrowser/Private/SAssetPicker.cpp
+++ b/Engine/Source/Editor/ContentBrowser/Private/SAssetPicker.cpp
@@ -12,6 +12,7 @@
 #include "Delegates/Delegate.h"
 #include "Editor.h"
 #include "Editor/EditorEngine.h"
+#include "Filters.h"
 #include "Framework/Application/SlateApplication.h"
 #include "Framework/Commands/GenericCommands.h"
 #include "Framework/Commands/UIAction.h"
@@ -315,6 +316,7 @@ void SAssetPicker::Construct( const FArguments& InArgs )
 		}
 	}
 
+	// clang-format off
 	VerticalBox->AddSlot()
 	.FillHeight(1.f)
 	[
@@ -331,6 +333,7 @@ void SAssetPicker::Construct( const FArguments& InArgs )
 		.OnVisualizeAssetToolTip(InArgs._AssetPickerConfig.OnVisualizeAssetToolTip)
 		.OnAssetToolTipClosing(InArgs._AssetPickerConfig.OnAssetToolTipClosing)
 		.FrontendFilters(FrontendFilters)
+	   	.ShowRedirectors_Lambda([this]() { return ContentBrowserUtils::ShouldShowRedirectors(FilterListPtr); })
 		.InitialSourcesData(CurrentSourcesData)
 		.InitialBackendFilter(CurrentBackendFilter)
 		.InitialViewType(InArgs._AssetPickerConfig.InitialAssetViewType)
@@ -360,7 +363,7 @@ void SAssetPicker::Construct( const FArguments& InArgs )
 		.OnSearchOptionsChanged(this, &SAssetPicker::HandleSearchSettingsChanged)
 		.InitialThumbnailSize(InArgs._AssetPickerConfig.InitialThumbnailSize)
 	];
-
+	// clang-format on
 
 	HorizontalBox->AddSlot()
 	.AutoWidth()
diff --git a/Engine/Source/Editor/ContentBrowser/Private/SAssetView.cpp b/Engine/Source/Editor/ContentBrowser/Private/SAssetView.cpp
--- a/Engine/Source/Editor/ContentBrowser/Private/SAssetView.cpp
+++ b/Engine/Source/Editor/ContentBrowser/Private/SAssetView.cpp
@@ -110,8 +110,14 @@ public:
 	explicit FAssetViewFrontendFilterHelper(SAssetView* InAssetView)
 		: AssetView(InAssetView)
 		, ContentBrowserData(IContentBrowserDataModule::Get().GetSubsystem())
+		, FolderFilter()
 		, bDisplayEmptyFolders(AssetView->IsShowingEmptyFolders())
 	{
+		if (!bDisplayEmptyFolders)
+		{
+			FolderFilter = FContentBrowserFolderContentsFilter{};
+			FolderFilter->ItemCategoryFilter = InAssetView->DetermineItemCategoryFilter();
+		}
 	}
 
 	bool DoesItemPassQueryFilter(const TSharedPtr<FAssetViewItem>& InItemToFilter)
@@ -146,7 +152,11 @@ public:
 		// Folders are only subject to "empty" filtering
 		if (InItemToFilter->IsFolder())
 		{
-			return ContentBrowserData->IsFolderVisible(InItemToFilter->GetItem().GetVirtualPath(), ContentBrowserUtils::GetIsFolderVisibleFlags(bDisplayEmptyFolders));
+			if (!ContentBrowserData->IsFolderVisible(InItemToFilter->GetItem().GetVirtualPath(), ContentBrowserUtils::GetIsFolderVisibleFlags(bDisplayEmptyFolders), FolderFilter))
+			{
+				return false;
+			}
+			return true;
 		}
 
 		// Run the item through the filters
@@ -161,6 +171,7 @@ public:
 private:
 	SAssetView* AssetView = nullptr;
 	UContentBrowserDataSubsystem* ContentBrowserData = nullptr;
+	TOptional<FContentBrowserFolderContentsFilter> FolderFilter;
 	const bool bDisplayEmptyFolders = true;
 };
 
@@ -199,6 +210,9 @@ void SAssetView::Construct( const FArguments& InArgs )
 	bFillEmptySpaceInTileView = InArgs._FillEmptySpaceInTileView;
 	FillScale = 1.0f;
 
+	bShowRedirectors = InArgs._ShowRedirectors;
+	bLastShowRedirectors = bShowRedirectors.Get(false);
+
 	ThumbnailHintFadeInSequence.JumpToStart();
 	ThumbnailHintFadeInSequence.AddCurve(0, 0.5f, ECurveEaseFunction::Linear);
 
@@ -1007,6 +1021,13 @@ void SAssetView::Tick( const FGeometry& AllottedGeometry, const double InCurrent
 		AssetThumbnailPool->Tick(InDeltaTime);
 	}
 
+	const bool bNewShowRedirectors = bShowRedirectors.Get(false);
+	if (bNewShowRedirectors != bLastShowRedirectors)
+	{
+		bLastShowRedirectors = bNewShowRedirectors;
+		OnFrontendFiltersChanged(); // refresh the same as if filters changed
+	}
+
 	CalculateThumbnailHintColorAndOpacity();
 
 	if (bPendingUpdateThumbnails)
@@ -1860,6 +1881,37 @@ bool SAssetView::IsValidSearchToken(const FString& Token) const
 	return true;
 }
 
+EContentBrowserItemCategoryFilter SAssetView::DetermineItemCategoryFilter() const
+{
+	// Check whether any legacy delegates are bound (the Content Browser doesn't use these, only pickers do)
+	// These limit the view to things that might use FAssetData
+	const bool bHasLegacyDelegateBindings = OnIsAssetValidForCustomToolTip.IsBound()
+										 || OnGetCustomAssetToolTip.IsBound()
+										 || OnVisualizeAssetToolTip.IsBound()
+										 || OnAssetToolTipClosing.IsBound()
+										 || OnShouldFilterAsset.IsBound();
+
+	EContentBrowserItemCategoryFilter ItemCategoryFilter = bHasLegacyDelegateBindings ? EContentBrowserItemCategoryFilter::IncludeAssets : InitialCategoryFilter;
+	if (IsShowingCppContent())
+	{
+		ItemCategoryFilter |= EContentBrowserItemCategoryFilter::IncludeClasses;
+	}
+	else
+	{
+		ItemCategoryFilter &= ~EContentBrowserItemCategoryFilter::IncludeClasses;
+	}
+	ItemCategoryFilter |= EContentBrowserItemCategoryFilter::IncludeCollections;
+	if (IsShowingRedirectors())
+	{
+		ItemCategoryFilter |= EContentBrowserItemCategoryFilter::IncludeRedirectors;
+	}
+	else
+	{
+		ItemCategoryFilter &= ~EContentBrowserItemCategoryFilter::IncludeRedirectors;
+	}
+	return ItemCategoryFilter;
+}
+
 FContentBrowserDataFilter SAssetView::CreateBackendDataFilter(bool bInvalidateCache) const
 {
 	// Assemble the filter using the current sources
@@ -1868,31 +1920,13 @@ FContentBrowserDataFilter SAssetView::CreateBackendDataFilter(bool bInvalidateCa
 	const bool bRecurse = ShouldFilterRecursively();
 	const bool bUsingFolders = IsShowingFolders() && !bRecurse;
 
-	// Check whether any legacy delegates are bound (the Content Browser doesn't use these, only pickers do)
-	// These limit the view to things that might use FAssetData
-	const bool bHasLegacyDelegateBindings 
-		=  OnIsAssetValidForCustomToolTip.IsBound()
-		|| OnGetCustomAssetToolTip.IsBound()
-		|| OnVisualizeAssetToolTip.IsBound()
-		|| OnAssetToolTipClosing.IsBound()
-		|| OnShouldFilterAsset.IsBound();
-
 	FContentBrowserDataFilter DataFilter;
 	DataFilter.bRecursivePaths = bRecurse || !bUsingFolders || bHasCollections;
 
 	DataFilter.ItemTypeFilter = EContentBrowserItemTypeFilter::IncludeFiles
 		| ((bUsingFolders && !bHasCollections) ? EContentBrowserItemTypeFilter::IncludeFolders : EContentBrowserItemTypeFilter::IncludeNone);
 
-	DataFilter.ItemCategoryFilter = bHasLegacyDelegateBindings ? EContentBrowserItemCategoryFilter::IncludeAssets : InitialCategoryFilter;
-	if (IsShowingCppContent())
-	{
-		DataFilter.ItemCategoryFilter |= EContentBrowserItemCategoryFilter::IncludeClasses;
-	}
-	else
-	{
-		DataFilter.ItemCategoryFilter &= ~EContentBrowserItemCategoryFilter::IncludeClasses;
-	}
-	DataFilter.ItemCategoryFilter |= EContentBrowserItemCategoryFilter::IncludeCollections;
+	DataFilter.ItemCategoryFilter = DetermineItemCategoryFilter();
 
 	DataFilter.ItemAttributeFilter = EContentBrowserItemAttributeFilter::IncludeProject
 		| (IsShowingEngineContent() ? EContentBrowserItemAttributeFilter::IncludeEngine : EContentBrowserItemAttributeFilter::IncludeNone)
@@ -2492,6 +2526,10 @@ void SAssetView::OnCollectionUpdated( const FCollectionNameType& Collection )
 
 void SAssetView::OnFrontendFiltersChanged()
 {
+	// We're refreshing so update the redirector visibility state in case it's not also bound to a frontend filter.
+	// This potentially avoids a double refresh on the next tick.
+	bLastShowRedirectors = bShowRedirectors.Get(false);
+
 	RequestQuickFrontendListRefresh();
 
 	// If we're not operating on recursively filtered data, we need to ensure a full slow
@@ -3043,6 +3081,11 @@ bool SAssetView::IsShowingEmptyFolders() const
 	return GetDefault<UContentBrowserSettings>()->DisplayEmptyFolders;
 }
 
+bool SAssetView::IsShowingRedirectors() const
+{
+	return bShowRedirectors.Get(false);
+}
+
 void SAssetView::ToggleRealTimeThumbnails()
 {
 	check(CanShowRealTimeThumbnails());
diff --git a/Engine/Source/Editor/ContentBrowser/Private/SContentBrowser.cpp b/Engine/Source/Editor/ContentBrowser/Private/SContentBrowser.cpp
--- a/Engine/Source/Editor/ContentBrowser/Private/SContentBrowser.cpp
+++ b/Engine/Source/Editor/ContentBrowser/Private/SContentBrowser.cpp
@@ -38,6 +38,7 @@
 #include "Editor.h"
 #include "Editor/EditorEngine.h"
 #include "FileHelpers.h"
+#include "Filters.h"
 #include "Filters/FilterBase.h"
 #include "Filters/SAssetFilterBar.h"
 #include "Filters/SBasicFilterBar.h"
@@ -378,34 +379,34 @@ void SContentBrowser::Construct( const FArguments& InArgs, const FName& InInstan
 	const bool bShowBottomToolbar = Config != nullptr ? Config->bShowBottomToolbar : true;
 
 	AssetViewPtr = SNew(SAssetView)
-		.ThumbnailLabel(Config != nullptr ? Config->ThumbnailLabel : EThumbnailLabel::ClassName)
-		//.ThumbnailScale(Config != nullptr ? Config->ThumbnailScale : 0.18f)
-		.InitialViewType(Config != nullptr ? Config->InitialAssetViewType : EAssetViewType::Tile)
-		.OnNewItemRequested(this, &SContentBrowser::OnNewItemRequested)
-		.OnItemSelectionChanged(this, &SContentBrowser::OnItemSelectionChanged, EContentBrowserViewContext::AssetView)
-		.OnItemsActivated(this, &SContentBrowser::OnItemsActivated)
-		.OnGetItemContextMenu(this, &SContentBrowser::GetItemContextMenu, EContentBrowserViewContext::AssetView)
-		.OnItemRenameCommitted(this, &SContentBrowser::OnItemRenameCommitted)
-		.OnShouldFilterItem(this, &SContentBrowser::HandlePrivateContentFilter)
-		.FrontendFilters(FrontendFilters)
-		.HighlightedText(this, &SContentBrowser::GetHighlightedText)
-		.ShowBottomToolbar(bShowBottomToolbar)
-		.ShowViewOptions(false)  // We control this for the main content browser
-		.AllowThumbnailEditMode(true)
-		.AllowThumbnailHintLabel(false)
-		.CanShowFolders(Config != nullptr ? Config->bCanShowFolders : true)
-		.CanShowClasses(Config != nullptr ? Config->bCanShowClasses : true)
-		.CanShowRealTimeThumbnails(Config != nullptr ? Config->bCanShowRealTimeThumbnails : true)
-		.CanShowDevelopersFolder(Config != nullptr ? Config->bCanShowDevelopersFolder : true)
-		.CanShowFavorites(true)
-		.CanDockCollections(true)
-		.AddMetaData<FTagMetaData>(FTagMetaData(TEXT("ContentBrowserAssets")))
-		.OwningContentBrowser(SharedThis(this))
-		.OnSearchOptionsChanged(this, &SContentBrowser::HandleAssetViewSearchOptionsChanged)
-		.bShowPathViewFilters(true)
-		.FillEmptySpaceInTileView(true)
-		.ShowDisallowedAssetClassAsUnsupportedItems(true);
-
+					   .ThumbnailLabel(Config != nullptr ? Config->ThumbnailLabel : EThumbnailLabel::ClassName)
+					   //.ThumbnailScale(Config != nullptr ? Config->ThumbnailScale : 0.18f)
+					   .InitialViewType(Config != nullptr ? Config->InitialAssetViewType : EAssetViewType::Tile)
+					   .OnNewItemRequested(this, &SContentBrowser::OnNewItemRequested)
+					   .OnItemSelectionChanged(this, &SContentBrowser::OnItemSelectionChanged, EContentBrowserViewContext::AssetView)
+					   .OnItemsActivated(this, &SContentBrowser::OnItemsActivated)
+					   .OnGetItemContextMenu(this, &SContentBrowser::GetItemContextMenu, EContentBrowserViewContext::AssetView)
+					   .OnItemRenameCommitted(this, &SContentBrowser::OnItemRenameCommitted)
+					   .OnShouldFilterItem(this, &SContentBrowser::HandlePrivateContentFilter)
+					   .FrontendFilters(FrontendFilters)
+					   .ShowRedirectors_Lambda([this]() { return ContentBrowserUtils::ShouldShowRedirectors(FilterListPtr); })
+					   .HighlightedText(this, &SContentBrowser::GetHighlightedText)
+					   .ShowBottomToolbar(bShowBottomToolbar)
+					   .ShowViewOptions(false) // We control this for the main content browser
+					   .AllowThumbnailEditMode(true)
+					   .AllowThumbnailHintLabel(false)
+					   .CanShowFolders(Config != nullptr ? Config->bCanShowFolders : true)
+					   .CanShowClasses(Config != nullptr ? Config->bCanShowClasses : true)
+					   .CanShowRealTimeThumbnails(Config != nullptr ? Config->bCanShowRealTimeThumbnails : true)
+					   .CanShowDevelopersFolder(Config != nullptr ? Config->bCanShowDevelopersFolder : true)
+					   .CanShowFavorites(true)
+					   .CanDockCollections(true)
+					   .AddMetaData<FTagMetaData>(FTagMetaData(TEXT("ContentBrowserAssets")))
+					   .OwningContentBrowser(SharedThis(this))
+					   .OnSearchOptionsChanged(this, &SContentBrowser::HandleAssetViewSearchOptionsChanged)
+					   .bShowPathViewFilters(true)
+					   .FillEmptySpaceInTileView(true)
+					   .ShowDisallowedAssetClassAsUnsupportedItems(true);
 
 	TSharedRef<SWidget> ViewOptions = SNullWidget::NullWidget;
 
@@ -1197,7 +1198,8 @@ void SContentBrowser::SetPathViewExpanded(bool bExpanded)
 }
 
 TSharedRef<SWidget> SContentBrowser::CreatePathView(const FContentBrowserConfig* Config)
-{	
+{
+	// clang-format off
 	return
 		SAssignNew(PathArea, SExpandableArea)
 		.BorderImage(FAppStyle::Get().GetBrush("Brushes.Header"))
@@ -1252,6 +1254,7 @@ TSharedRef<SWidget> SContentBrowser::CreatePathView(const FContentBrowserConfig*
 					.FocusSearchBoxWhenOpened(false)
 					.ShowTreeTitle(false)
 					.ShowSeparator(false)
+				   	.ShowRedirectors_Lambda([this]() { return ContentBrowserUtils::ShouldShowRedirectors(FilterListPtr); })
 					.AllowClassesFolder(true)
 					.AddMetaData<FTagMetaData>(FTagMetaData(TEXT("ContentBrowserSources")))
 					.ExternalSearch(SourcesSearch)
@@ -1260,6 +1263,7 @@ TSharedRef<SWidget> SContentBrowser::CreatePathView(const FContentBrowserConfig*
 				]
 			]
 		];
+	// clang-format on
 }
 
 TSharedRef<SWidget> SContentBrowser::CreateDockedCollectionsView(const FContentBrowserConfig* Config)
@@ -1634,6 +1638,13 @@ void SContentBrowser::PrepareToSyncItems(TArrayView<const FContentBrowserItem> I
  			}
 		}
 	}
+	
+	if (bDisableFiltersThatHideAssets)
+	{
+		// Disable the filter categories
+		// Do this before repopulate because the redirectors filter can hide folders
+		FilterListPtr->DisableFiltersThatHideItems(ItemsToSync);
+	}
 
 	// If we have auto-enabled any flags or found a non-existant path, force a refresh
 	if (bRepopulate)
@@ -1677,12 +1688,6 @@ void SContentBrowser::PrepareToSyncItems(TArrayView<const FContentBrowserItem> I
 		FavoritePathViewPtr->Populate();
 	}
 
-	if ( bDisableFiltersThatHideAssets )
-	{
-		// Disable the filter categories
-		FilterListPtr->DisableFiltersThatHideItems(ItemsToSync);
-	}
-
 	// Disable the filter search (reset the filter, then clear the search text)
 	// Note: we have to remove the filter immediately, we can't wait for OnSearchBoxChanged to hit
 	SetSearchBoxText(FText::GetEmpty());
diff --git a/Engine/Source/Editor/ContentBrowser/Private/SFilterList.cpp b/Engine/Source/Editor/ContentBrowser/Private/SFilterList.cpp
--- a/Engine/Source/Editor/ContentBrowser/Private/SFilterList.cpp
+++ b/Engine/Source/Editor/ContentBrowser/Private/SFilterList.cpp
@@ -3,6 +3,7 @@
 
 #include "SFilterList.h"
 
+#include "Algo/AnyOf.h"
 #include "AssetRegistry/ARFilter.h"
 #include "ContentBrowserDataFilter.h"
 #include "ContentBrowserDataSource.h"
@@ -12,6 +13,7 @@
 #include "ContentBrowserItemData.h"
 #include "ContentBrowserMenuContexts.h"
 #include "ContentBrowserUtils.h"
+#include "Filters.h"
 #include "Filters/FilterBarConfig.h"
 #include "Filters/SAssetFilterBar.h"
 #include "Framework/Application/MenuStack.h"
@@ -75,7 +77,7 @@ void SFilterList::Construct( const FArguments& InArgs )
 	AllFrontendFilters_Internal.Add( MakeShareable(new FFrontendFilter_Writable(DefaultCategory)) );
 	AllFrontendFilters_Internal.Add( MakeShareable(new FFrontendFilter_ShowOtherDevelopers(DefaultCategory)) );
 	AllFrontendFilters_Internal.Add( MakeShareable(new FFrontendFilter_ReplicatedBlueprint(DefaultCategory)) );
-	AllFrontendFilters_Internal.Add( MakeShareable(new FFrontendFilter_ShowRedirectors(DefaultCategory)) );
+	AllFrontendFilters_Internal.Add(MakeShared<FFilter_ShowRedirectors>(DefaultCategory));
 	AllFrontendFilters_Internal.Add( MakeShareable(new FFrontendFilter_InUseByLoadedLevels(DefaultCategory)) );
 	AllFrontendFilters_Internal.Add( MakeShareable(new FFrontendFilter_UsedInAnyLevel(DefaultCategory)) );
 	AllFrontendFilters_Internal.Add( MakeShareable(new FFrontendFilter_NotUsedInAnyLevel(DefaultCategory)) );
@@ -252,6 +254,26 @@ void SFilterList::DisableFiltersThatHideItems(TArrayView<const FContentBrowserIt
 			}
 		}
 
+		// Special case: if the object is a redirector then enable the 'show redirectors' filter - this will also prevent
+		// folders that contain only redirectors from being hidden with the "hide empty folders" setting
+		FString RedirectorClassPath = UObjectRedirector::StaticClass()->GetPathName();
+		const bool bAnyRedirectors = Algo::AnyOf(ItemList, [RedirectorClassPath](const FContentBrowserItem& Item) {
+			FContentBrowserItemDataAttributeValue Attribute = Item.GetItemAttribute(ContentBrowserItemAttributes::ItemTypeName, false);
+			return Attribute.IsValid() && Attribute.GetValueString() == RedirectorClassPath;
+		});
+		if (bAnyRedirectors)
+		{
+			TSharedPtr<FFilter_ShowRedirectors> RedirectorFilter = StaticCastSharedPtr<FFilter_ShowRedirectors>(GetFrontendFilter("ShowRedirectorsBackend"));
+			if (RedirectorFilter.IsValid())
+			{
+				int32 ExistingIndex = Filters.IndexOfByPredicate([RedirectorFilter](TSharedPtr<SFilter> Filter) { return Filter->GetFrontendFilter() == RedirectorFilter; });
+				TSharedRef<SFilter> FilterWidget = ExistingIndex == INDEX_NONE ? AddFilterToBar(RedirectorFilter.ToSharedRef()) : Filters[ExistingIndex];
+				FilterWidget->SetEnabled(true, false);
+				SetFrontendFilterActive(RedirectorFilter.ToSharedRef(), true);
+				ExecuteOnFilterChanged = true;
+			}
+		}
+
 		// Disable all backend filters if it was determined that the combined backend filter hides any of the assets
 		if (bDisableAllBackendFilters)
 		{
diff --git a/Engine/Source/Editor/ContentBrowser/Private/SPathView.cpp b/Engine/Source/Editor/ContentBrowser/Private/SPathView.cpp
--- a/Engine/Source/Editor/ContentBrowser/Private/SPathView.cpp
+++ b/Engine/Source/Editor/ContentBrowser/Private/SPathView.cpp
@@ -143,6 +143,8 @@ void SPathView::Construct( const FArguments& InArgs )
 	InitialCategoryFilter = InArgs._InitialCategoryFilter;
 	bAllowClassesFolder = InArgs._AllowClassesFolder;
 	bAllowReadOnlyFolders = InArgs._AllowReadOnlyFolders;
+	bShowRedirectors = InArgs._ShowRedirectors;
+	bLastShowRedirectors = bShowRedirectors.Get(false);
 	PreventTreeItemChangedDelegateCount = 0;
 	TreeTitle = LOCTEXT("AssetTreeTitle", "Asset Tree");
 	if ( InArgs._FocusSearchBoxWhenOpened )
@@ -333,6 +335,19 @@ void SPathView::Construct( const FArguments& InArgs )
 	}
 }
 
+void SPathView::Tick(const FGeometry& AllottedGeometry, const double InCurrentTime, const float InDeltaTime)
+{
+	Super::Tick(AllottedGeometry, InCurrentTime, InDeltaTime);
+
+	const bool bNewShowRedirectors = bShowRedirectors.Get(false);
+	if (bNewShowRedirectors != bLastShowRedirectors)
+	{
+		UE_LOG(LogContentBrowser, Verbose, TEXT("PathView bShowRedirectors changed to %d"), bNewShowRedirectors);
+		bLastShowRedirectors = bNewShowRedirectors;
+		HandleSettingChanged("ShowRedirectors");
+	}
+}
+
 void SPathView::PopulatePathViewFiltersMenu(UToolMenu* Menu)
 {
 	{
@@ -934,6 +949,14 @@ EContentBrowserItemCategoryFilter SPathView::GetContentBrowserItemCategoryFilter
 		ItemCategoryFilter &= ~EContentBrowserItemCategoryFilter::IncludeClasses;
 	}
 	ItemCategoryFilter &= ~EContentBrowserItemCategoryFilter::IncludeCollections;
+	if (bShowRedirectors.Get(false))
+	{
+		ItemCategoryFilter |= EContentBrowserItemCategoryFilter::IncludeRedirectors;
+	}
+	else
+	{
+		ItemCategoryFilter &= ~EContentBrowserItemCategoryFilter::IncludeRedirectors;
+	}
 
 	return ItemCategoryFilter;
 }
@@ -1520,6 +1543,7 @@ FText SPathView::GetHighlightText() const
 void SPathView::Populate(const bool bIsRefreshingFilter)
 {
 	TRACE_CPUPROFILER_EVENT_SCOPE(SPathView::Populate);
+	UE_LOG(LogContentBrowser, Verbose, TEXT("Repopulating path view"));
 
 	// Update the list of expanded path before removing the items
 	UpdateLastExpandedPathsIfDirty();
@@ -1545,23 +1569,36 @@ void SPathView::Populate(const bool bIsRefreshingFilter)
 
 		UContentBrowserDataSubsystem* ContentBrowserData = IContentBrowserDataModule::Get().GetSubsystem();
 		const FContentBrowserDataCompiledFilter CompiledDataFilter = CreateCompiledFolderFilter();
+		TOptional<FContentBrowserFolderContentsFilter> FolderFilter;
+		if (!bDisplayEmpty)
+		{
+			FolderFilter = FContentBrowserFolderContentsFilter{};
+			FolderFilter->ItemCategoryFilter = CompiledDataFilter.ItemCategoryFilter;
+		}
+		EContentBrowserIsFolderVisibleFlags FolderFlags = ContentBrowserUtils::GetIsFolderVisibleFlags(bDisplayEmpty);
 
 		TArray<TSharedPtr<FTreeItem>> ItemsCreated;
-		ContentBrowserData->EnumerateItemsMatchingFilter(CompiledDataFilter, [this, bFilteringByText, bDisplayEmpty, ContentBrowserData, &ItemsCreated](FContentBrowserItemData&& InItemData)
-		{
-			bool bPassesFilter = ContentBrowserData->IsFolderVisible(InItemData.GetVirtualPath(), ContentBrowserUtils::GetIsFolderVisibleFlags(bDisplayEmpty));
-			if (bPassesFilter && bFilteringByText)
-			{
-				// Use the whole path so we deliberately include any children of matched parents in the filtered list
-				const FString PathStr = InItemData.GetVirtualPath().ToString();
-				bPassesFilter &= SearchBoxFolderFilter->PassesFilter(PathStr);
-			}
+		ContentBrowserData->EnumerateItemsMatchingFilter(CompiledDataFilter,
+			[this, bFilteringByText, bDisplayEmpty, CompiledDataFilter, FolderFilter, FolderFlags, ContentBrowserData, &ItemsCreated](FContentBrowserItemData&& InItemData) {
+				if (!ContentBrowserData->IsFolderVisible(InItemData.GetVirtualPath(), FolderFlags, FolderFilter))
+				{
+					UE_LOG(LogContentBrowser, VeryVerbose, TEXT("Hiding folder %s that fails current pre-text filtering"), *WriteToString<256>(InItemData.GetVirtualPath()));
+					return true; // continue enumerating
+				}
+
+				if (bFilteringByText)
+				{
+					// Use the whole path so we deliberately include any children of matched parents in the filtered list
+					const FString PathStr = InItemData.GetVirtualPath().ToString();
+					if (!SearchBoxFolderFilter->PassesFilter(PathStr))
+					{
+						return true; // continue enumerating
+					}
+				}
 
-			if (bPassesFilter)
-			{
 				// Using array of all items created to handle item expansion of fully virtual paths that may not be included in enumeration
 				ItemsCreated.Reset();
-				AddFolderItem(MoveTemp(InItemData), /*bUserNamed=*/ false, &ItemsCreated);
+				AddFolderItem(MoveTemp(InItemData), /*bUserNamed=*/false, &ItemsCreated);
 
 				for (TSharedPtr<FTreeItem> Item : ItemsCreated)
 				{
@@ -1593,10 +1630,9 @@ void SPathView::Populate(const bool bIsRefreshingFilter)
 						}
 					}
 				}
-			}
 
-			return true;
-		});
+				return true;
+			});
 	}
 
 	SortRootItems();
@@ -2002,16 +2038,23 @@ void SPathView::HandleItemDataUpdated(TArrayView<const FContentBrowserItemDataUp
 		}
 	};
 
-	auto DoesItemPassFilter = [this, bFilteringByText, bDisplayEmpty, ContentBrowserData, &CompiledDataFilter](const FContentBrowserItemData& InItemData)
+	TOptional<FContentBrowserFolderContentsFilter> FolderFilter;
+	if (!bDisplayEmpty)
 	{
+		FolderFilter = FContentBrowserFolderContentsFilter{};
+		FolderFilter->ItemCategoryFilter = CompiledDataFilter.ItemCategoryFilter;
+	}
+	EContentBrowserIsFolderVisibleFlags FolderFlags = ContentBrowserUtils::GetIsFolderVisibleFlags(bDisplayEmpty);
+	auto DoesItemPassFilter = [this, bFilteringByText, FolderFlags, FolderFilter, ContentBrowserData, &CompiledDataFilter](const FContentBrowserItemData& InItemData) {
 		UContentBrowserDataSource* ItemDataSource = InItemData.GetOwnerDataSource();
 		if (!ItemDataSource->DoesItemPassFilter(InItemData, CompiledDataFilter))
 		{
 			return false;
 		}
 
-		if (!ContentBrowserData->IsFolderVisible(InItemData.GetVirtualPath(), ContentBrowserUtils::GetIsFolderVisibleFlags(bDisplayEmpty)))
+		if (!ContentBrowserData->IsFolderVisible(InItemData.GetVirtualPath(), FolderFlags, FolderFilter))
 		{
+			UE_LOG(LogContentBrowser, VeryVerbose, TEXT("Hiding folder %s that fails broad visibility filter"), *WriteToString<256>(InItemData.GetVirtualPath()));
 			return false;
 		}
 
@@ -2114,14 +2157,15 @@ bool SPathView::PathIsFilteredFromViewBySearch(const FString& InPath) const
 
 void SPathView::HandleSettingChanged(FName PropertyName)
 {
-	if ((PropertyName == GET_MEMBER_NAME_CHECKED(UContentBrowserSettings, DisplayEmptyFolders)) ||
-		(PropertyName == "DisplayDevelopersFolder") ||
-		(PropertyName == "DisplayEngineFolder") ||
-		(PropertyName == "DisplayPluginFolders") ||
-		(PropertyName == "DisplayL10NFolder") ||
-		(PropertyName == GET_MEMBER_NAME_CHECKED(UContentBrowserSettings, bDisplayContentFolderSuffix)) ||
-		(PropertyName == GET_MEMBER_NAME_CHECKED(UContentBrowserSettings, bDisplayFriendlyNameForPluginFolders)) ||
-		(PropertyName == NAME_None))	// @todo: Needed if PostEditChange was called manually, for now
+	if ((PropertyName == GET_MEMBER_NAME_CHECKED(UContentBrowserSettings, DisplayEmptyFolders))
+		|| (PropertyName == "ShowRedirectors")
+		|| (PropertyName == "DisplayDevelopersFolder")
+		|| (PropertyName == "DisplayEngineFolder")
+		|| (PropertyName == "DisplayPluginFolders")
+		|| (PropertyName == "DisplayL10NFolder")
+		|| (PropertyName == GET_MEMBER_NAME_CHECKED(UContentBrowserSettings, bDisplayContentFolderSuffix))
+		|| (PropertyName == GET_MEMBER_NAME_CHECKED(UContentBrowserSettings, bDisplayFriendlyNameForPluginFolders))
+		|| (PropertyName == NAME_None)) // @todo: Needed if PostEditChange was called manually, for now
 	{
 		const bool bHadSelectedPath = TreeViewPtr->GetNumItemsSelected() > 0;
 
diff --git a/Engine/Source/Editor/ContentBrowser/Private/SPathView.h b/Engine/Source/Editor/ContentBrowser/Private/SPathView.h
--- a/Engine/Source/Editor/ContentBrowser/Private/SPathView.h
+++ b/Engine/Source/Editor/ContentBrowser/Private/SPathView.h
@@ -62,6 +62,8 @@ typedef TTextFilter< const FString& > FolderTextFilter;
  */
 class SPathView : public SCompoundWidget
 {
+	using Super = SCompoundWidget;
+
 public:
 	/** Delegate for when plugin filters have changed */
 	DECLARE_DELEGATE( FOnFrontendPluginFilterChanged );
@@ -115,6 +117,9 @@ public:
 		/** If true, the favorites expander will be displayed */
 		SLATE_ARGUMENT(bool, ShowFavorites);
 
+		/** If true, redirectors are taken into consideration when deciding if folders are empty */
+		SLATE_ATTRIBUTE(bool, ShowRedirectors);
+
 		/** The selection mode for the tree view */
 		SLATE_ARGUMENT( ESelectionMode::Type, SelectionMode )
 
@@ -138,6 +143,9 @@ public:
 	/** Constructs this widget with InArgs */
 	virtual void Construct( const FArguments& InArgs );
 
+	/** Tick to poll attributes */
+	virtual void Tick(const FGeometry& AllottedGeometry, const double InCurrentTime, const float InDeltaTime) override;
+
 	/** Selects the closest matches to the supplied paths in the tree. "/" delimited */
 	void SetSelectedPaths(const TArray<FName>& Paths);
 
@@ -446,6 +454,9 @@ protected:
 	/** Custom Folder permissions */
 	TSharedPtr<FPathPermissionList> CustomFolderPermissionList;
 
+	TAttribute<bool> bShowRedirectors;
+	bool bLastShowRedirectors;
+
 private:
 	/** Used to track if the list of last expanded path should be updated */
 	bool bLastExpandedPathsDirty = false;
diff --git a/Engine/Source/Editor/ContentBrowser/Public/FrontendFilters.h b/Engine/Source/Editor/ContentBrowser/Public/FrontendFilters.h
--- a/Engine/Source/Editor/ContentBrowser/Public/FrontendFilters.h
+++ b/Engine/Source/Editor/ContentBrowser/Public/FrontendFilters.h
@@ -265,7 +265,8 @@ private:
 };
 
 /** An inverse filter that allows display of object redirectors */
-class CONTENTBROWSER_API FFrontendFilter_ShowRedirectors : public FFrontendFilter
+class UE_DEPRECATED(5.5, "FFrontendFilter_ShowRedirectors is deprecated. FFilter_ShowRedirectors is a virtual filter which controls backend search state instead.")
+	CONTENTBROWSER_API FFrontendFilter_ShowRedirectors : public FFrontendFilter
 {
 public:
 	/** Constructor */
diff --git a/Engine/Source/Editor/ContentBrowser/Public/SAssetView.h b/Engine/Source/Editor/ContentBrowser/Public/SAssetView.h
--- a/Engine/Source/Editor/ContentBrowser/Public/SAssetView.h
+++ b/Engine/Source/Editor/ContentBrowser/Public/SAssetView.h
@@ -50,6 +50,7 @@ class FAssetViewItem;
 class FContentBrowserItemDataTemporaryContext;
 class FContentBrowserItemDataUpdate;
 class FDragDropEvent;
+class FFilter_ShowRedirectors;
 class FMenuBuilder;
 class FPathPermissionList;
 class FSlateRect;
@@ -180,6 +181,9 @@ public:
 		/** The filter collection used to further filter down assets returned from the backend */
 		SLATE_ARGUMENT( TSharedPtr<FAssetFilterCollectionType>, FrontendFilters )
 
+		/** If true, redirectors are visible even if not explicitly searching for them. */
+		SLATE_ATTRIBUTE(bool, ShowRedirectors);
+
 		/** Show path view filters submenu in view options menu */
 		SLATE_ARGUMENT( bool, bShowPathViewFilters )
 
@@ -543,6 +547,9 @@ private:
 	/** @return true when we are showing empty folders */
 	bool IsShowingEmptyFolders() const;
 
+	/** @return true when the asset view is showing object redirectors */
+	bool IsShowingRedirectors() const;
+
 	/** Toggle whether localized content should be shown or not */
 	void ToggleShowLocalizedContent();
 
@@ -862,6 +869,7 @@ private:
 	/** Append the current effective backend filter (intersection of BackendFilter and SupportedFilter) to the given filter. */
 	void AppendBackendFilter(FARFilter& FilterToAppendTo) const;
 
+	EContentBrowserItemCategoryFilter DetermineItemCategoryFilter() const;
 	FContentBrowserDataFilter CreateBackendDataFilter(bool bInvalidateCache) const;
 
 	/** Handles updating the view when content items are changed */
@@ -921,6 +929,9 @@ private:
 	TSharedPtr<FPathPermissionList> WritableFolderPermissionList;
 	TSharedPtr<FAssetFilterCollectionType> FrontendFilters;
 
+	TAttribute<bool> bShowRedirectors;
+	bool bLastShowRedirectors;
+
 	/** Show path view filters submenu in view options menu  */
 	bool bShowPathViewFilters;
 
diff --git a/Engine/Source/Editor/ContentBrowserData/Private/ContentBrowserDataSource.cpp b/Engine/Source/Editor/ContentBrowserData/Private/ContentBrowserDataSource.cpp
--- a/Engine/Source/Editor/ContentBrowserData/Private/ContentBrowserDataSource.cpp
+++ b/Engine/Source/Editor/ContentBrowserData/Private/ContentBrowserDataSource.cpp
@@ -254,7 +254,7 @@ bool UContentBrowserDataSource::PrioritizeSearchPath(const FName InPath)
 	return false;
 }
 
-bool UContentBrowserDataSource::IsFolderVisible(const FName InPath, const EContentBrowserIsFolderVisibleFlags InFlags)
+bool UContentBrowserDataSource::IsFolderVisible(const FName Path, const EContentBrowserIsFolderVisibleFlags Flags, TOptional<FContentBrowserFolderContentsFilter> ContentsFilter)
 {
 	return true;
 }
diff --git a/Engine/Source/Editor/ContentBrowserData/Private/ContentBrowserDataSubsystem.cpp b/Engine/Source/Editor/ContentBrowserData/Private/ContentBrowserDataSubsystem.cpp
--- a/Engine/Source/Editor/ContentBrowserData/Private/ContentBrowserDataSubsystem.cpp
+++ b/Engine/Source/Editor/ContentBrowserData/Private/ContentBrowserDataSubsystem.cpp
@@ -537,7 +537,7 @@ bool UContentBrowserDataSubsystem::PrioritizeSearchPath(const FName InPath)
 	return bDidPrioritize;
 }
 
-bool UContentBrowserDataSubsystem::IsFolderVisible(const FName InPath, const EContentBrowserIsFolderVisibleFlags InFlags) const
+bool UContentBrowserDataSubsystem::IsFolderVisible(const FName InPath, const EContentBrowserIsFolderVisibleFlags InFlags, TOptional<FContentBrowserFolderContentsFilter> InContentsFilter) const
 {
 	bool bIsKnownPath = false;
 	for (const auto& ActiveDataSourcePair : ActiveDataSources)
@@ -546,7 +546,7 @@ bool UContentBrowserDataSubsystem::IsFolderVisible(const FName InPath, const ECo
 		if (DataSource->IsVirtualPathUnderMountRoot(InPath))
 		{
 			bIsKnownPath = true;
-			if (DataSource->IsFolderVisible(InPath, InFlags))
+			if (DataSource->IsFolderVisible(InPath, InFlags, InContentsFilter))
 			{
 				return true;
 			}
@@ -559,7 +559,9 @@ bool UContentBrowserDataSubsystem::IsFolderVisible(const FName InPath, const ECo
 
 bool UContentBrowserDataSubsystem::IsFolderVisibleIfHidingEmpty(const FName InPath) const
 {
+	PRAGMA_DISABLE_DEPRECATION_WARNINGS
 	return IsFolderVisible(InPath, EContentBrowserIsFolderVisibleFlags::Default | EContentBrowserIsFolderVisibleFlags::HideEmptyFolders);
+	PRAGMA_ENABLE_DEPRECATION_WARNINGS
 }
 
 bool UContentBrowserDataSubsystem::CanCreateFolder(const FName InPath, FText* OutErrorMsg) const
diff --git a/Engine/Source/Editor/ContentBrowserData/Public/ContentBrowserDataFilter.h b/Engine/Source/Editor/ContentBrowserData/Public/ContentBrowserDataFilter.h
--- a/Engine/Source/Editor/ContentBrowserData/Public/ContentBrowserDataFilter.h
+++ b/Engine/Source/Editor/ContentBrowserData/Public/ContentBrowserDataFilter.h
@@ -25,7 +25,10 @@ class FPathPermissionList;
 class UContentBrowserDataSource;
 class UContentBrowserDataSubsystem;
 
-/** Flags controlling which item types should be included */
+/**
+ * Flags controlling which item types should be included
+ * Generally matches EContentBrowserItemFlags::Type_MASK
+ */
 UENUM(Flags)
 enum class EContentBrowserItemTypeFilter : uint8
 {
@@ -36,7 +39,10 @@ enum class EContentBrowserItemTypeFilter : uint8
 };
 ENUM_CLASS_FLAGS(EContentBrowserItemTypeFilter);
 
-/** Flags controlling which item categories should be included */
+/**
+ * Flags controlling which item categories should be included
+ * Generally matches EContentBrowserItemFlags::Category_MASK
+ */
 UENUM(Flags)
 enum class EContentBrowserItemCategoryFilter : uint8
 {
@@ -44,8 +50,9 @@ enum class EContentBrowserItemCategoryFilter : uint8
 	IncludeAssets = 1<<0,
 	IncludeClasses = 1<<1,
 	IncludeCollections = 1<<2,
-	IncludeMisc = 1<<3,
-	IncludeAll = IncludeAssets | IncludeClasses | IncludeCollections | IncludeMisc,
+	IncludeRedirectors = 1 << 3,
+	IncludeMisc = 1 << 4,
+	IncludeAll = IncludeAssets | IncludeClasses | IncludeCollections | IncludeRedirectors | IncludeMisc,
 };
 ENUM_CLASS_FLAGS(EContentBrowserItemCategoryFilter);
 
@@ -340,6 +347,24 @@ public:
 	bool bIncludeChildCollections = false;
 };
 
+/*
+ * Structure used to optionally filter folders by their broad contents
+ */
+USTRUCT(BlueprintType)
+struct CONTENTBROWSERDATA_API FContentBrowserFolderContentsFilter
+{
+	GENERATED_BODY()
+
+public:
+	UPROPERTY(EditAnywhere, BlueprintReadWrite, Category = "ContentBrowser")
+	EContentBrowserItemTypeFilter ItemTypeFilter = EContentBrowserItemTypeFilter::IncludeAll;
+
+	UPROPERTY(EditAnywhere, BlueprintReadWrite, Category = "ContentBrowser")
+	EContentBrowserItemCategoryFilter ItemCategoryFilter = EContentBrowserItemCategoryFilter::IncludeAll;
+
+	UPROPERTY(EditAnywhere, BlueprintReadWrite, Category = "ContentBrowser")
+	EContentBrowserItemAttributeFilter ItemAttributeFilter = EContentBrowserItemAttributeFilter::IncludeAll;
+};
 
 /**
  * Data used to tell the content browser to show the item that doesn't pass the class permission list as a unsupported asset
diff --git a/Engine/Source/Editor/ContentBrowserData/Public/ContentBrowserDataSource.h b/Engine/Source/Editor/ContentBrowserData/Public/ContentBrowserDataSource.h
--- a/Engine/Source/Editor/ContentBrowserData/Public/ContentBrowserDataSource.h
+++ b/Engine/Source/Editor/ContentBrowserData/Public/ContentBrowserDataSource.h
@@ -256,8 +256,11 @@ public:
 	/**
 	 * Query whether the given virtual folder should be visible in the UI.
 	 * @note This function must be able to answer the question quickly or not at all (and assume visible). It *must not* block doing something like a file system scan.
+	 * @param Path Virtual path of folder e.g. /All/GameData/Stuff
+	 * @param Flags Currently unused
+	 * @param ContentsFilter Optional filter for limiting visibility to only folders with a certain type of contents based on desired view filtering. If this is unset, folders may be visible regardless of their contents or empty status.
 	 */
-	virtual bool IsFolderVisible(const FName InPath, const EContentBrowserIsFolderVisibleFlags InFlags);
+	virtual bool IsFolderVisible(const FName Path, const EContentBrowserIsFolderVisibleFlags Flags, TOptional<FContentBrowserFolderContentsFilter> ContentsFilter);
 
 	/*
 	 * Query whether a folder can be created at the given virtual path, optionally providing error information if it cannot.
diff --git a/Engine/Source/Editor/ContentBrowserData/Public/ContentBrowserDataSubsystem.h b/Engine/Source/Editor/ContentBrowserData/Public/ContentBrowserDataSubsystem.h
--- a/Engine/Source/Editor/ContentBrowserData/Public/ContentBrowserDataSubsystem.h
+++ b/Engine/Source/Editor/ContentBrowserData/Public/ContentBrowserDataSubsystem.h
@@ -55,7 +55,7 @@ enum class EContentBrowserIsFolderVisibleFlags : uint8
 	/**
 	 * Hide folders that recursively contain no file items.
 	 */
-	HideEmptyFolders = 1<<0,
+	HideEmptyFolders UE_DEPRECATED(5.5, "Empty folder filtering is now dependent on which types of contents are visible, use the FContentBrowserFolderContentsFilter argument to IsFolderVisible") = 1 << 0,
 
 	/**
 	 * Default visibility flags.
@@ -253,8 +253,11 @@ public:
 
 	/**
 	 * Query whether the given virtual folder should be visible in the UI.
+	 * @param Path Virtual path of folder e.g. /All/GameData/Stuff
+	 * @param Flags Currently unused
+	 * @param ContentsFilter Optional filter for limiting visibility to only folders with a certain type of contents based on desired view filtering.
 	 */
-	bool IsFolderVisible(const FName InPath, const EContentBrowserIsFolderVisibleFlags InFlags = EContentBrowserIsFolderVisibleFlags::Default) const;
+	bool IsFolderVisible(const FName Path, const EContentBrowserIsFolderVisibleFlags Flags = EContentBrowserIsFolderVisibleFlags::Default, TOptional<FContentBrowserFolderContentsFilter> ContentsFilter = {}) const;
 
 	/**
 	 * Query whether the given virtual folder should be visible if the UI is asking to hide empty content folders.
