Drastically improve the performance of the function UContentBrowserAssetDataSource::IsKnownContentPath. On a large project this can reduce the length of the hitches during the asset registry asset discovery period.

#tests Validated by comparing the returned values of the old implementation against the new one.

diff --git a/Engine/Plugins/Editor/ContentBrowser/ContentBrowserAssetDataSource/Source/ContentBrowserAssetDataSource/Private/ContentBrowserAssetDataSource.cpp b/Engine/Plugins/Editor/ContentBrowser/ContentBrowserAssetDataSource/Source/ContentBrowserAssetDataSource/Private/ContentBrowserAssetDataSource.cpp
--- a/Engine/Plugins/Editor/ContentBrowser/ContentBrowserAssetDataSource/Source/ContentBrowserAssetDataSource/Private/ContentBrowserAssetDataSource.cpp
+++ b/Engine/Plugins/Editor/ContentBrowser/ContentBrowserAssetDataSource/Source/ContentBrowserAssetDataSource/Private/ContentBrowserAssetDataSource.cpp
@@ -17,6 +17,7 @@
 #include "Editor.h"
 #include "IAssetTools.h"
 #include "ToolMenus.h"
+#include "Misc/Char.h"
 #include "Misc/PackageName.h"
 #include "Misc/PathViews.h"
 #include "NewAssetContextMenu.h"
@@ -247,10 +248,13 @@ void UContentBrowserAssetDataSource::Initialize(const bool InAutoRegister)
 
 	BuildRootPathVirtualTree();
 
-	// Mount roots are always visible
 	for (const FString& RootContentPath : RootContentPaths)
 	{
+		// Mount roots are always visible
 		OnAlwaysShowPath(RootContentPath);
+		
+		// Populate the acceleration structure
+		AddRootContentPathToStateMachine(RootContentPath);
 	}
 }
 
@@ -261,6 +265,9 @@ void UContentBrowserAssetDataSource::Shutdown()
 	AssetTools = nullptr;
 	AssetRegistry = nullptr;
 
+	RootContentPaths.Empty();
+	RootContentPathsTrie.NextNodes.Empty();
+
 	if (FAssetRegistryModule* AssetRegistryModule = FModuleManager::GetModulePtr<FAssetRegistryModule>(AssetRegistryConstants::ModuleName))
 	{
 		IAssetRegistry* AssetRegistryMaybe = AssetRegistryModule->TryGet();
@@ -2783,28 +2790,36 @@ bool UContentBrowserAssetDataSource::IsKnownContentPath(const FName InPackagePat
 {
 	FNameBuilder PackagePathStr(InPackagePath);
 	const FStringView PackagePathStrView = PackagePathStr;
-	for (const FString& RootContentPath : RootContentPaths)
+
+	const FCharacterNode* CurrentNode = &RootContentPathsTrie;
+
+	for (const TCHAR& Character : PackagePathStrView)
 	{
-		const FStringView RootContentPathNoSlash = FStringView(RootContentPath).LeftChop(1);
-		if (PackagePathStrView.StartsWith(RootContentPath, ESearchCase::IgnoreCase) || PackagePathStrView.Equals(RootContentPathNoSlash, ESearchCase::IgnoreCase))
+		const TPair<FCharacterNodePtr, int32>* NextNodePair = CurrentNode->NextNodes.Find(TChar<TCHAR>::ToLower(Character));
+
+		if (!NextNodePair)
 		{
-			return true;
+			// This text start with no root content path
+			return false;
 		}
-	}
-
-	return false;
-}
+		 
+		const FCharacterNode* NextNode = NextNodePair->Key.Get();
 
-bool UContentBrowserAssetDataSource::IsRootContentPath(const FName InPackagePath) const
-{
-	FNameBuilder PackagePathStr(InPackagePath);
-	PackagePathStr << TEXT('/'); // RootContentPaths have a trailing slash
+		// Is the next node terminal
+		if (NextNode->NextNodes.IsEmpty())
+		{
+			// The package path start with a root content path
+			return true;
+		}
 
-	const FStringView PackagePathStrView = PackagePathStr;
-	return RootContentPaths.ContainsByPredicate([&PackagePathStrView](const FString& InRootContentPath)
-	{
-		return PackagePathStrView == InRootContentPath;
-	});
+		CurrentNode = NextNode;
+	}
+	
+	/**
+	 * Test if the folder is a root folder here like / Game.
+	 * Where the only thing missing is the last '/'.
+	 */
+	return CurrentNode->NextNodes.Contains(TEXT('/'));
 }
 
 bool UContentBrowserAssetDataSource::GetObjectPathsForCollections(ICollectionManager* CollectionManager, TArrayView<const FCollectionNameType> InCollections, const bool bIncludeChildCollections, TArray<FSoftObjectPath>& OutObjectPaths)
@@ -3114,7 +3129,8 @@ void UContentBrowserAssetDataSource::BuildRootPathVirtualTree()
 
 void UContentBrowserAssetDataSource::OnContentPathMounted(const FString& InAssetPath, const FString& InFileSystemPath)
 {
-	RootContentPaths.AddUnique(InAssetPath);
+	RootContentPaths.Add(InAssetPath);
+	AddRootContentPathToStateMachine(InAssetPath);
 
 	RootPathAdded(InAssetPath);
 
@@ -3126,6 +3142,7 @@ void UContentBrowserAssetDataSource::OnContentPathDismounted(const FString& InAs
 {
 	RootPathRemoved(InAssetPath);
 
+	RemoveRootContentPathFromStateMachine(InAssetPath);
 	RootContentPaths.Remove(InAssetPath);
 }
 
@@ -3520,6 +3537,50 @@ FContentBrowserItemData UContentBrowserAssetDataSource::OnFinalizeDuplicateAsset
 	return CreateAssetFileItem(FAssetData(Asset));
 }
 
+void UContentBrowserAssetDataSource::AddRootContentPathToStateMachine(const FString& InAssetPath)
+{
+	/**
+	 * No need to mark the nodes with a terminal attribute on the last node since they always finish with a '/'.
+	 * They never contains more then two '/'. One at the start and the other at end.
+	 */
+	 ensure(InAssetPath[InAssetPath.Len() - 1] == '/') ;
+
+	FCharacterNode* CurrentNode = &RootContentPathsTrie;
+
+	for (const TCHAR& Character : InAssetPath)
+	{
+		TPair<FCharacterNodePtr,int32>& NextNode = CurrentNode->NextNodes.FindOrAdd(TChar<TCHAR>::ToLower(Character));
+		++NextNode.Value;
+		CurrentNode = NextNode.Key.Get();
+	}
+}
+
+void UContentBrowserAssetDataSource::RemoveRootContentPathFromStateMachine(const FString& InAssetPath)
+{
+	FCharacterNode* CurrentNode = &RootContentPathsTrie;
+
+	for (const TCHAR& Character : InAssetPath)
+	{
+		const TCHAR LoweredCharacter = TChar<TCHAR>::ToLower(Character);
+		uint32 Hash = GetTypeHash(LoweredCharacter);
+		TPair<FCharacterNodePtr,int32>* NextNode = CurrentNode->NextNodes.FindByHash(Hash,LoweredCharacter);
+
+		if (!NextNode)
+		{
+			return;
+		}
+
+		--NextNode->Value;
+		if (NextNode->Value == 0)
+		{
+			CurrentNode->NextNodes.RemoveByHash(Hash,LoweredCharacter);
+			return;
+		}
+
+		CurrentNode = NextNode->Key.Get();
+	}
+}
+
 bool UContentBrowserAssetDataSource::PathPassesCompiledDataFilter(const FContentBrowserCompiledAssetDataFilter& InFilter, const FName InInternalPath)
 {
 	if (InFilter.ExcludedPackagePaths.Contains(InInternalPath)) // PassesExcludedPathsFilter
diff --git a/Engine/Plugins/Editor/ContentBrowser/ContentBrowserAssetDataSource/Source/ContentBrowserAssetDataSource/Public/ContentBrowserAssetDataSource.h b/Engine/Plugins/Editor/ContentBrowser/ContentBrowserAssetDataSource/Source/ContentBrowserAssetDataSource/Public/ContentBrowserAssetDataSource.h
--- a/Engine/Plugins/Editor/ContentBrowser/ContentBrowserAssetDataSource/Source/ContentBrowserAssetDataSource/Public/ContentBrowserAssetDataSource.h
+++ b/Engine/Plugins/Editor/ContentBrowser/ContentBrowserAssetDataSource/Source/ContentBrowserAssetDataSource/Public/ContentBrowserAssetDataSource.h
@@ -358,8 +358,6 @@ protected:
 private:
 	bool IsKnownContentPath(const FName InPackagePath) const;
 
-	bool IsRootContentPath(const FName InPackagePath) const;
-
 	static bool GetObjectPathsForCollections(ICollectionManager* CollectionManager, TArrayView<const FCollectionNameType> InCollections, const bool bIncludeChildCollections, TArray<FSoftObjectPath>& OutObjectPaths);
 
 	FContentBrowserItemData CreateAssetFolderItem(const FName InFolderPath);
@@ -442,6 +440,10 @@ private:
 
 	FContentBrowserItemData OnFinalizeDuplicateAsset(const FContentBrowserItemData& InItemData, const FString& InProposedName, FText* OutErrorMsg);
 
+	void AddRootContentPathToStateMachine(const FString& InAssetPath);
+
+	void RemoveRootContentPathFromStateMachine(const FString& InAssetPath);
+
 	IAssetRegistry* AssetRegistry;
 
 	IAssetTools* AssetTools;
@@ -460,6 +462,65 @@ private:
 	 */
 	TArray<FString> RootContentPaths;
 
+	struct FCharacterNode;
+
+	struct FCharacterNodePtr
+	{
+		FCharacterNodePtr()
+			: Node(MakeUnique<FCharacterNode>())
+		{
+		}
+
+		FCharacterNodePtr(FCharacterNodePtr&&) = default;
+		FCharacterNodePtr& operator=(FCharacterNodePtr&&) = default;
+
+		FCharacterNodePtr(const FCharacterNodePtr&) = delete;
+		FCharacterNodePtr& operator=(const FCharacterNodePtr&) = delete;
+
+		const FCharacterNode& operator*() const
+		{
+			return Node.operator*();
+		}
+
+		const FCharacterNode* operator->() const
+		{
+			return Node.operator->();
+		}
+
+		FCharacterNode& operator*()
+		{
+			return Node.operator*();
+		}
+
+		FCharacterNode* operator->()
+		{
+			return Node.operator->();
+		}
+ 
+		const FCharacterNode* Get() const
+		{
+			return Node.Get();
+		}
+
+		FCharacterNode* Get()
+		{
+			return Node.Get();
+		}
+
+	private:
+
+		TUniquePtr<FCharacterNode> Node;
+	};
+
+	struct FCharacterNode
+	{
+		// The next characters in the tree and the number of paths beginning with the prefix including that character for use in removing paths.
+		TMap<TCHAR, TPair<FCharacterNodePtr, int32>> NextNodes;
+	};
+
+	// Tree of character nodes all in lower case. Used to speed up queries against the RootContentPaths Array.
+	FCharacterNode RootContentPathsTrie;
+
 	/**
 	 * Map of folders that have attributes set.
 	 */
