Cache internal path in content browser data and use internal and virtual path fields to determine invariant path

diff --git a/Engine/Plugins/Editor/ContentBrowser/ContentBrowserAliasDataSoure/Source/ContentBrowserAliasDataSource/Private/ContentBrowserAliasDataSource.cpp b/Engine/Plugins/Editor/ContentBrowser/ContentBrowserAliasDataSoure/Source/ContentBrowserAliasDataSource/Private/ContentBrowserAliasDataSource.cpp
--- a/Engine/Plugins/Editor/ContentBrowser/ContentBrowserAliasDataSoure/Source/ContentBrowserAliasDataSource/Private/ContentBrowserAliasDataSource.cpp
+++ b/Engine/Plugins/Editor/ContentBrowser/ContentBrowserAliasDataSoure/Source/ContentBrowserAliasDataSource/Private/ContentBrowserAliasDataSource.cpp
@@ -987,14 +987,20 @@ void UContentBrowserAliasDataSource::ClearCachedFilterData(const FContentBrowser
 	FilterCache.ClearCachedData(IDOwner);
 }
 
-FContentBrowserItemData UContentBrowserAliasDataSource::CreateAssetFolderItem(const FName InFolderPath)
+FContentBrowserItemData UContentBrowserAliasDataSource::CreateAssetFolderItem(const FName InInternalFolderPath)
 {
 	FName VirtualizedPath;
-	TryConvertInternalPathToVirtual(InFolderPath, VirtualizedPath);
+	TryConvertInternalPathToVirtual(InInternalFolderPath, VirtualizedPath);
 
-	const FString FolderItemName = FPackageName::GetShortName(InFolderPath);
-	return FContentBrowserItemData(this, EContentBrowserItemFlags::Type_Folder | EContentBrowserItemFlags::Category_Asset, VirtualizedPath,
-		*FolderItemName, AliasFolderDisplayNames.FindRef(InFolderPath), MakeShared<FContentBrowserAssetFolderItemDataPayload>(InFolderPath));
+	const FString FolderItemName = FPackageName::GetShortName(InInternalFolderPath);
+	return FContentBrowserItemData(
+		this,
+		EContentBrowserItemFlags::Type_Folder | EContentBrowserItemFlags::Category_Asset,
+		VirtualizedPath,
+		*FolderItemName,
+		AliasFolderDisplayNames.FindRef(InInternalFolderPath),
+		MakeShared<FContentBrowserAssetFolderItemDataPayload>(InInternalFolderPath),
+		{ InInternalFolderPath });
 }
 
 FContentBrowserItemData UContentBrowserAliasDataSource::CreateAssetFileItem(const FContentBrowserUniqueAlias& Alias)
@@ -1005,15 +1011,21 @@ FContentBrowserItemData UContentBrowserAliasDataSource::CreateAssetFileItem(cons
 	{
 		LLM_SCOPE(ELLMTag::UI);
 		FName VirtualizedPath;
-PRAGMA_DISABLE_DEPRECATION_WARNINGS
-		TryConvertInternalPathToVirtual(AliasData->ObjectPath.ToFName(), VirtualizedPath);
-PRAGMA_ENABLE_DEPRECATION_WARNINGS
+		PRAGMA_DISABLE_DEPRECATION_WARNINGS
+		FName InternalPath = AliasData->ObjectPath.ToFName();
+		PRAGMA_ENABLE_DEPRECATION_WARNINGS
+		TryConvertInternalPathToVirtual(InternalPath, VirtualizedPath);
 
 		// Since AliasID is PackagePath/AssetName, AssetName should also be passed as the ItemName here. This provides the functionality of
 		// being able to have multiple aliases with the same display name, while still showing their original asset name in the tooltip.
-		return FContentBrowserItemData(this, EContentBrowserItemFlags::Type_File | EContentBrowserItemFlags::Category_Asset, VirtualizedPath,
-			AliasData->AssetData.AssetName, AliasData->AliasDisplayName, MakeShared<FContentBrowserAliasItemDataPayload>(AliasData->AssetData, Alias));
+		return FContentBrowserItemData(
+			this,
+			EContentBrowserItemFlags::Type_File | EContentBrowserItemFlags::Category_Asset,
+			VirtualizedPath,
+			AliasData->AssetData.AssetName,
+			AliasData->AliasDisplayName,
+			MakeShared<FContentBrowserAliasItemDataPayload>(AliasData->AssetData, Alias),
+			{ InternalPath });
 	}
 	return FContentBrowserItemData();
 }
-
diff --git a/Engine/Plugins/Editor/ContentBrowser/ContentBrowserAssetDataSource/Source/ContentBrowserAssetDataSource/Private/ContentBrowserAssetDataCore.cpp b/Engine/Plugins/Editor/ContentBrowser/ContentBrowserAssetDataSource/Source/ContentBrowserAssetDataSource/Private/ContentBrowserAssetDataCore.cpp
--- a/Engine/Plugins/Editor/ContentBrowser/ContentBrowserAssetDataSource/Source/ContentBrowserAssetDataSource/Private/ContentBrowserAssetDataCore.cpp
+++ b/Engine/Plugins/Editor/ContentBrowser/ContentBrowserAssetDataSource/Source/ContentBrowserAssetDataSource/Private/ContentBrowserAssetDataCore.cpp
@@ -36,31 +36,60 @@ DEFINE_LOG_CATEGORY(LogContentBrowserAssetDataSource);
 namespace ContentBrowserAssetData
 {
 
-FContentBrowserItemData CreateAssetFolderItem(UContentBrowserDataSource* InOwnerDataSource, const FName InVirtualPath, const FName InFolderPath, const bool bIsCookedPath, const bool bIsPlugin)
-{
-	const FString FolderItemName = FPackageName::GetShortName(InFolderPath);
-	FText FolderDisplayNameOverride = ContentBrowserDataUtils::GetFolderItemDisplayNameOverride(InFolderPath, FolderItemName, /*bIsClassesFolder*/ false, bIsCookedPath);
-	return FContentBrowserItemData(InOwnerDataSource,
-		EContentBrowserItemFlags::Type_Folder | EContentBrowserItemFlags::Category_Asset | (bIsPlugin ? EContentBrowserItemFlags::Category_Plugin : EContentBrowserItemFlags::None ),
+FContentBrowserItemData CreateAssetFolderItem(
+	UContentBrowserDataSource* InOwnerDataSource,
+	const FName InVirtualPath,
+	const FName InInternalFolderPath,
+	const bool bIsCookedPath,
+	const bool bIsPlugin)
+{
+	const FString FolderItemName = FPackageName::GetShortName(InInternalFolderPath);
+	FText FolderDisplayNameOverride = ContentBrowserDataUtils::GetFolderItemDisplayNameOverride(
+		InInternalFolderPath, FolderItemName, /*bIsClassesFolder*/ false, bIsCookedPath);
+	return FContentBrowserItemData(
+		InOwnerDataSource,
+		EContentBrowserItemFlags::Type_Folder | EContentBrowserItemFlags::Category_Asset
+			| (bIsPlugin ? EContentBrowserItemFlags::Category_Plugin : EContentBrowserItemFlags::None),
 		InVirtualPath,
 		*FolderItemName,
 		MoveTemp(FolderDisplayNameOverride),
-		MakeShared<FContentBrowserAssetFolderItemDataPayload>(InFolderPath));
+		MakeShared<FContentBrowserAssetFolderItemDataPayload>(InInternalFolderPath),
+		{ InInternalFolderPath });
 }
 
-FContentBrowserItemData CreateAssetFileItem(UContentBrowserDataSource* InOwnerDataSource, const FName InVirtualPath, const FAssetData& InAssetData, const bool bIsPluginAsset)
+FContentBrowserItemData CreateAssetFileItem(
+	UContentBrowserDataSource* InOwnerDataSource,
+	const FName InVirtualPath,
+	const FName InInternalPath,
+	const FAssetData& InAssetData,
+	const bool bIsPluginAsset)
 {
-	return FContentBrowserItemData(InOwnerDataSource,
-		EContentBrowserItemFlags::Type_File | EContentBrowserItemFlags::Category_Asset | (bIsPluginAsset ? EContentBrowserItemFlags::Category_Plugin : EContentBrowserItemFlags::None),
+	return FContentBrowserItemData(
+		InOwnerDataSource,
+		EContentBrowserItemFlags::Type_File | EContentBrowserItemFlags::Category_Asset
+			| (bIsPluginAsset ? EContentBrowserItemFlags::Category_Plugin : EContentBrowserItemFlags::None),
 		InVirtualPath,
 		InAssetData.AssetName,
 		FText(),
-		MakeShared<FContentBrowserAssetFileItemDataPayload>(InAssetData));
+		MakeShared<FContentBrowserAssetFileItemDataPayload>(InAssetData),
+		{ InInternalPath });
 }
 
-FContentBrowserItemData CreateUnsupportedAssetFileItem(UContentBrowserDataSource* InOwnerDataSource, const FName InVirtualPath, const FAssetData& InAssetData)
+FContentBrowserItemData CreateUnsupportedAssetFileItem(
+	UContentBrowserDataSource* InOwnerDataSource,
+	const FName InVirtualPath,
+	const FName InInternalPath,
+	const FAssetData& InAssetData)
 {
-	return FContentBrowserItemData(InOwnerDataSource, EContentBrowserItemFlags::Type_File | EContentBrowserItemFlags::Category_Asset | EContentBrowserItemFlags::Misc_Unsupported, InVirtualPath, InAssetData.AssetName, FText(), MakeShared<FContentBrowserUnsupportedAssetFileItemDataPayload>(InAssetData));
+	return FContentBrowserItemData(
+		InOwnerDataSource,
+		EContentBrowserItemFlags::Type_File | EContentBrowserItemFlags::Category_Asset
+			| EContentBrowserItemFlags::Misc_Unsupported,
+		InVirtualPath,
+		InAssetData.AssetName,
+		FText(),
+		MakeShared<FContentBrowserUnsupportedAssetFileItemDataPayload>(InAssetData),
+		{ InInternalPath });
 }
 
 
@@ -1931,7 +1960,10 @@ void PopulateAssetFileContextMenu(UContentBrowserDataSource* InOwnerDataSource,
 				FName VirtualPathToShow;
 				if (OwnerDataSourcePtr->Legacy_TryConvertAssetDataToVirtualPath(AssetToShow, /*bUseFolderPaths*/false, VirtualPathToShow))
 				{
-					ItemsToShow.Emplace(CreateAssetFileItem(OwnerDataSourcePtr, VirtualPathToShow, AssetToShow));
+					PRAGMA_DISABLE_DEPRECATION_WARNINGS
+					ItemsToShow.Emplace(
+						CreateAssetFileItem(OwnerDataSourcePtr, VirtualPathToShow, AssetToShow.ObjectPath, AssetToShow));
+					PRAGMA_ENABLE_DEPRECATION_WARNINGS
 				}
 			}
 			OnShowInPathsView.Execute(ItemsToShow);
@@ -1945,6 +1977,6 @@ void PopulateAssetFileContextMenu(UContentBrowserDataSource* InOwnerDataSource,
 	);
 }
 
-}
+} // namespace ContentBrowserAssetData
 
 #undef LOCTEXT_NAMESPACE
diff --git a/Engine/Plugins/Editor/ContentBrowser/ContentBrowserAssetDataSource/Source/ContentBrowserAssetDataSource/Private/ContentBrowserAssetDataSource.cpp b/Engine/Plugins/Editor/ContentBrowser/ContentBrowserAssetDataSource/Source/ContentBrowserAssetDataSource/Private/ContentBrowserAssetDataSource.cpp
--- a/Engine/Plugins/Editor/ContentBrowser/ContentBrowserAssetDataSource/Source/ContentBrowserAssetDataSource/Private/ContentBrowserAssetDataSource.cpp
+++ b/Engine/Plugins/Editor/ContentBrowser/ContentBrowserAssetDataSource/Source/ContentBrowserAssetDataSource/Private/ContentBrowserAssetDataSource.cpp
@@ -2384,24 +2384,27 @@ bool UContentBrowserAssetDataSource::DuplicateItem(const FContentBrowserItemData
 	if (ContentBrowserAssetData::DuplicateItem(AssetTools, this, InItem, SourceAsset, NewAssetData))
 	{
 		FName VirtualizedPath;
-PRAGMA_DISABLE_DEPRECATION_WARNINGS
-		TryConvertInternalPathToVirtual(NewAssetData.ObjectPath, VirtualizedPath);
-PRAGMA_ENABLE_DEPRECATION_WARNINGS
+		PRAGMA_DISABLE_DEPRECATION_WARNINGS
+		FName InternalPath = NewAssetData.ObjectPath;
+		PRAGMA_ENABLE_DEPRECATION_WARNINGS
+		TryConvertInternalPathToVirtual(InternalPath, VirtualizedPath);
 
 		FContentBrowserItemData NewItemData(
 			this,
-			EContentBrowserItemFlags::Type_File | EContentBrowserItemFlags::Category_Asset | EContentBrowserItemFlags::Temporary_Duplication,
+			EContentBrowserItemFlags::Type_File | EContentBrowserItemFlags::Category_Asset
+				| EContentBrowserItemFlags::Temporary_Duplication,
 			VirtualizedPath,
 			NewAssetData.AssetName,
 			FText::AsCultureInvariant(NewAssetData.AssetName.ToString()),
-			MakeShared<FContentBrowserAssetFileItemDataPayload_Duplication>(MoveTemp(NewAssetData), SourceAsset)
-			);
+			MakeShared<FContentBrowserAssetFileItemDataPayload_Duplication>(MoveTemp(NewAssetData), SourceAsset),
+			{ InternalPath });
 
 		OutPendingItem = FContentBrowserItemDataTemporaryContext(
 			MoveTemp(NewItemData),
-			FContentBrowserItemDataTemporaryContext::FOnValidateItem::CreateUObject(this, &UContentBrowserAssetDataSource::OnValidateItemName),
-			FContentBrowserItemDataTemporaryContext::FOnFinalizeItem::CreateUObject(this, &UContentBrowserAssetDataSource::OnFinalizeDuplicateAsset)
-			);
+			FContentBrowserItemDataTemporaryContext::FOnValidateItem::CreateUObject(
+				this, &UContentBrowserAssetDataSource::OnValidateItemName),
+			FContentBrowserItemDataTemporaryContext::FOnFinalizeItem::CreateUObject(
+				this, &UContentBrowserAssetDataSource::OnFinalizeDuplicateAsset));
 
 		return true;
 	}
@@ -2839,41 +2842,47 @@ bool UContentBrowserAssetDataSource::GetObjectPathsForCollections(ICollectionMan
 	return false;
 }
 
-FContentBrowserItemData UContentBrowserAssetDataSource::CreateAssetFolderItem(const FName InFolderPath)
+FContentBrowserItemData UContentBrowserAssetDataSource::CreateAssetFolderItem(const FName InInternalFolderPath)
 {
 	FName VirtualizedPath;
-	TryConvertInternalPathToVirtual(InFolderPath, VirtualizedPath);
+	TryConvertInternalPathToVirtual(InInternalFolderPath, VirtualizedPath);
 
-	const EContentBrowserFolderAttributes FolderAttributes = GetAssetFolderAttributes(InFolderPath);
-	const bool bIsCookedPath = EnumHasAnyFlags(FolderAttributes, EContentBrowserFolderAttributes::HasAssets | EContentBrowserFolderAttributes::HasRedirectors)
-							&& !EnumHasAnyFlags(FolderAttributes, EContentBrowserFolderAttributes::HasSourceContent);
+	const EContentBrowserFolderAttributes FolderAttributes = GetAssetFolderAttributes(InInternalFolderPath);
+	const bool bIsCookedPath =
+		EnumHasAnyFlags(
+			FolderAttributes, EContentBrowserFolderAttributes::HasAssets | EContentBrowserFolderAttributes::HasRedirectors)
+		&& !EnumHasAnyFlags(FolderAttributes, EContentBrowserFolderAttributes::HasSourceContent);
 	const bool bIsPlugin = EnumHasAnyFlags(FolderAttributes, EContentBrowserFolderAttributes::IsInPlugin);
-	return ContentBrowserAssetData::CreateAssetFolderItem(this, VirtualizedPath, InFolderPath, bIsCookedPath, bIsPlugin);
+	return ContentBrowserAssetData::CreateAssetFolderItem(
+		this, VirtualizedPath, InInternalFolderPath, bIsCookedPath, bIsPlugin);
 }
 
 FContentBrowserItemData UContentBrowserAssetDataSource::CreateAssetFileItem(const FAssetData& InAssetData)
 {
 	FName VirtualizedPath;
-PRAGMA_DISABLE_DEPRECATION_WARNINGS
-	TryConvertInternalPathToVirtual(InAssetData.ObjectPath, VirtualizedPath);
-PRAGMA_ENABLE_DEPRECATION_WARNINGS
+	PRAGMA_DISABLE_DEPRECATION_WARNINGS
+	FName InternalPath = InAssetData.ObjectPath;
+	PRAGMA_ENABLE_DEPRECATION_WARNINGS
+	TryConvertInternalPathToVirtual(InternalPath, VirtualizedPath);
 
 	const EContentBrowserFolderAttributes FolderAttributes = GetAssetFolderAttributes(InAssetData.PackagePath);
 	const bool bIsPlugin = EnumHasAnyFlags(FolderAttributes, EContentBrowserFolderAttributes::IsInPlugin);
-	return ContentBrowserAssetData::CreateAssetFileItem(this, VirtualizedPath, InAssetData, bIsPlugin);
+	return ContentBrowserAssetData::CreateAssetFileItem(this, VirtualizedPath, InternalPath, InAssetData, bIsPlugin);
 }
 
 FContentBrowserItemData UContentBrowserAssetDataSource::CreateUnsupportedAssetFileItem(const FAssetData& InAssetData)
 {
 	FName VirtualizedPath;
 	PRAGMA_DISABLE_DEPRECATION_WARNINGS
-		TryConvertInternalPathToVirtual(InAssetData.ObjectPath, VirtualizedPath);
+	FName InternalPath = InAssetData.ObjectPath;
 	PRAGMA_ENABLE_DEPRECATION_WARNINGS
+	TryConvertInternalPathToVirtual(InternalPath, VirtualizedPath);
 
-	return ContentBrowserAssetData::CreateUnsupportedAssetFileItem(this, VirtualizedPath, InAssetData);
+	return ContentBrowserAssetData::CreateUnsupportedAssetFileItem(this, VirtualizedPath, InternalPath, InAssetData);
 }
 
-TSharedPtr<const FContentBrowserAssetFolderItemDataPayload> UContentBrowserAssetDataSource::GetAssetFolderItemPayload(const FContentBrowserItemData& InItem) const
+TSharedPtr<const FContentBrowserAssetFolderItemDataPayload> UContentBrowserAssetDataSource::GetAssetFolderItemPayload(
+	const FContentBrowserItemData& InItem) const
 {
 	return ContentBrowserAssetData::GetAssetFolderItemPayload(this, InItem);
 }
@@ -3418,24 +3427,27 @@ void UContentBrowserAssetDataSource::OnBeginCreateAsset(const FName InDefaultAss
 		FAssetData NewAssetData(*(InPackagePath.ToString() / InDefaultAssetName.ToString()), InPackagePath, InDefaultAssetName, ClassToUse->GetClassPathName());
 
 		FName VirtualizedPath;
-PRAGMA_DISABLE_DEPRECATION_WARNINGS
-		TryConvertInternalPathToVirtual(NewAssetData.ObjectPath, VirtualizedPath);
-PRAGMA_ENABLE_DEPRECATION_WARNINGS
+		PRAGMA_DISABLE_DEPRECATION_WARNINGS
+		FName InternalPath = NewAssetData.ObjectPath;
+		PRAGMA_ENABLE_DEPRECATION_WARNINGS
+		TryConvertInternalPathToVirtual(InternalPath, VirtualizedPath);
 
 		FContentBrowserItemData NewItemData(
-			this, 
-			EContentBrowserItemFlags::Type_File | EContentBrowserItemFlags::Category_Asset | EContentBrowserItemFlags::Temporary_Creation,
-			VirtualizedPath, 
+			this,
+			EContentBrowserItemFlags::Type_File | EContentBrowserItemFlags::Category_Asset
+				| EContentBrowserItemFlags::Temporary_Creation,
+			VirtualizedPath,
 			NewAssetData.AssetName,
 			FText::AsCultureInvariant(NewAssetData.AssetName.ToString()),
-			MakeShared<FContentBrowserAssetFileItemDataPayload_Creation>(MoveTemp(NewAssetData), InAssetClass, InFactory)
-			);
+			MakeShared<FContentBrowserAssetFileItemDataPayload_Creation>(MoveTemp(NewAssetData), InAssetClass, InFactory),
+			{ InternalPath });
 
 		InOnBeginItemCreation.Execute(FContentBrowserItemDataTemporaryContext(
-			MoveTemp(NewItemData), 
-			FContentBrowserItemDataTemporaryContext::FOnValidateItem::CreateUObject(this, &UContentBrowserAssetDataSource::OnValidateItemName),
-			FContentBrowserItemDataTemporaryContext::FOnFinalizeItem::CreateUObject(this, &UContentBrowserAssetDataSource::OnFinalizeCreateAsset)
-			));
+			MoveTemp(NewItemData),
+			FContentBrowserItemDataTemporaryContext::FOnValidateItem::CreateUObject(
+				this, &UContentBrowserAssetDataSource::OnValidateItemName),
+			FContentBrowserItemDataTemporaryContext::FOnFinalizeItem::CreateUObject(
+				this, &UContentBrowserAssetDataSource::OnFinalizeCreateAsset)));
 	}
 }
 
diff --git a/Engine/Plugins/Editor/ContentBrowser/ContentBrowserAssetDataSource/Source/ContentBrowserAssetDataSource/Public/ContentBrowserAssetDataCore.h b/Engine/Plugins/Editor/ContentBrowser/ContentBrowserAssetDataSource/Source/ContentBrowserAssetDataSource/Public/ContentBrowserAssetDataCore.h
--- a/Engine/Plugins/Editor/ContentBrowser/ContentBrowserAssetDataSource/Source/ContentBrowserAssetDataSource/Public/ContentBrowserAssetDataCore.h
+++ b/Engine/Plugins/Editor/ContentBrowser/ContentBrowserAssetDataSource/Source/ContentBrowserAssetDataSource/Public/ContentBrowserAssetDataCore.h
@@ -22,11 +22,25 @@ DECLARE_LOG_CATEGORY_EXTERN(LogContentBrowserAssetDataSource, Log, All);
 namespace ContentBrowserAssetData
 {
 
-	CONTENTBROWSERASSETDATASOURCE_API FContentBrowserItemData CreateAssetFolderItem(UContentBrowserDataSource* InOwnerDataSource, const FName InVirtualPath, const FName InFolderPath, const bool bIsCookedPath = false, const bool bIsPlugin = false);
-
-	CONTENTBROWSERASSETDATASOURCE_API FContentBrowserItemData CreateAssetFileItem(UContentBrowserDataSource* InOwnerDataSource, const FName InVirtualPath, const FAssetData& InAssetData, const bool bIsPlugin = false);
-
-	CONTENTBROWSERASSETDATASOURCE_API FContentBrowserItemData CreateUnsupportedAssetFileItem(UContentBrowserDataSource* InOwnerDataSource, const FName InVirtualPath, const FAssetData& InAssetData);
+CONTENTBROWSERASSETDATASOURCE_API FContentBrowserItemData CreateAssetFolderItem(
+	UContentBrowserDataSource* InOwnerDataSource,
+	const FName InVirtualPath,
+	const FName InInternalFolderPath,
+	const bool bIsCookedPath = false,
+	const bool bIsPlugin = false);
+
+CONTENTBROWSERASSETDATASOURCE_API FContentBrowserItemData CreateAssetFileItem(
+	UContentBrowserDataSource* InOwnerDataSource,
+	const FName InVirtualPath,
+	const FName InInternalPath,
+	const FAssetData& InAssetData,
+	const bool bIsPlugin = false);
+
+CONTENTBROWSERASSETDATASOURCE_API FContentBrowserItemData CreateUnsupportedAssetFileItem(
+	UContentBrowserDataSource* InOwnerDataSource,
+	const FName InVirtualPath,
+	const FName InInternalPath,
+	const FAssetData& InAssetData);
 
 	CONTENTBROWSERASSETDATASOURCE_API TSharedPtr<const FContentBrowserAssetFolderItemDataPayload> GetAssetFolderItemPayload(const UContentBrowserDataSource* InOwnerDataSource, const FContentBrowserItemData& InItem);
 
diff --git a/Engine/Plugins/Editor/ContentBrowser/ContentBrowserAssetDataSource/Source/ContentBrowserAssetDataSource/Public/ContentBrowserAssetDataSource.h b/Engine/Plugins/Editor/ContentBrowser/ContentBrowserAssetDataSource/Source/ContentBrowserAssetDataSource/Public/ContentBrowserAssetDataSource.h
--- a/Engine/Plugins/Editor/ContentBrowser/ContentBrowserAssetDataSource/Source/ContentBrowserAssetDataSource/Public/ContentBrowserAssetDataSource.h
+++ b/Engine/Plugins/Editor/ContentBrowser/ContentBrowserAssetDataSource/Source/ContentBrowserAssetDataSource/Public/ContentBrowserAssetDataSource.h
@@ -360,7 +360,7 @@ private:
 
 	static bool GetObjectPathsForCollections(ICollectionManager* CollectionManager, TArrayView<const FCollectionNameType> InCollections, const bool bIncludeChildCollections, TArray<FSoftObjectPath>& OutObjectPaths);
 
-	FContentBrowserItemData CreateAssetFolderItem(const FName InFolderPath);
+	FContentBrowserItemData CreateAssetFolderItem(const FName InInternalFolderPath);
 
 	FContentBrowserItemData CreateAssetFileItem(const FAssetData& InAssetData);
 
diff --git a/Engine/Plugins/Editor/ContentBrowser/ContentBrowserClassDataSource/Source/ContentBrowserClassDataSource/Private/ContentBrowserClassDataCore.cpp b/Engine/Plugins/Editor/ContentBrowser/ContentBrowserClassDataSource/Source/ContentBrowserClassDataSource/Private/ContentBrowserClassDataCore.cpp
--- a/Engine/Plugins/Editor/ContentBrowser/ContentBrowserClassDataSource/Source/ContentBrowserClassDataSource/Private/ContentBrowserClassDataCore.cpp
+++ b/Engine/Plugins/Editor/ContentBrowser/ContentBrowserClassDataSource/Source/ContentBrowserClassDataSource/Private/ContentBrowserClassDataCore.cpp
@@ -100,7 +100,8 @@ bool IsPluginClass(const FName InPath)
 	return AssetViewUtils::IsPluginFolder(PathStr);
 }
 
-FContentBrowserItemData CreateClassFolderItem(UContentBrowserDataSource* InOwnerDataSource, const FName InVirtualPath, const FName InFolderPath, const bool bIsFromPlugin)
+FContentBrowserItemData CreateClassFolderItem(
+	UContentBrowserDataSource* InOwnerDataSource, const FName InVirtualPath, FName InFolderPath, const bool bIsFromPlugin)
 {
 	static const FName GameRootPath = "/Classes_Game";
 	static const FName EngineRootPath = "/Classes_Engine";
@@ -126,20 +127,28 @@ FContentBrowserItemData CreateClassFolderItem(UContentBrowserDataSource* InOwner
 		InVirtualPath,
 		*FolderItemName,
 		MoveTemp(FolderDisplayNameOverride),
-		MakeShared<FContentBrowserClassFolderItemDataPayload>(InFolderPath));
+		MakeShared<FContentBrowserClassFolderItemDataPayload>(InFolderPath),
+		{ InFolderPath });
 }
 
-FContentBrowserItemData CreateClassFileItem(UContentBrowserDataSource* InOwnerDataSource, const FName InVirtualPath, const FName InClassPath, UClass* InClass, const bool bIsFromPlugin)
+FContentBrowserItemData CreateClassFileItem(
+	UContentBrowserDataSource* InOwnerDataSource,
+	const FName InVirtualPath,
+	FName InClassPath,
+	UClass* InClass,
+	const bool bIsFromPlugin)
 {
 	return FContentBrowserItemData(InOwnerDataSource,
 		EContentBrowserItemFlags::Type_File | EContentBrowserItemFlags::Category_Class | (bIsFromPlugin ? EContentBrowserItemFlags::Category_Plugin : EContentBrowserItemFlags::None),
 		InVirtualPath,
 		InClass->GetFName(),
 		FText(),
-		MakeShared<FContentBrowserClassFileItemDataPayload>(InClassPath, InClass));
+		MakeShared<FContentBrowserClassFileItemDataPayload>(InClassPath, InClass),
+		{ InClassPath });
 }
 
-TSharedPtr<const FContentBrowserClassFolderItemDataPayload> GetClassFolderItemPayload(const UContentBrowserDataSource* InOwnerDataSource, const FContentBrowserItemData& InItem)
+TSharedPtr<const FContentBrowserClassFolderItemDataPayload> GetClassFolderItemPayload(
+	const UContentBrowserDataSource* InOwnerDataSource, const FContentBrowserItemData& InItem)
 {
 	if (InItem.GetOwnerDataSource() == InOwnerDataSource && InItem.IsFolder())
 	{
diff --git a/Engine/Plugins/Editor/ContentBrowser/ContentBrowserFileDataSource/Source/ContentBrowserFileDataSource/Private/ContentBrowserFileDataCore.cpp b/Engine/Plugins/Editor/ContentBrowser/ContentBrowserFileDataSource/Source/ContentBrowserFileDataSource/Private/ContentBrowserFileDataCore.cpp
--- a/Engine/Plugins/Editor/ContentBrowser/ContentBrowserFileDataSource/Source/ContentBrowserFileDataSource/Private/ContentBrowserFileDataCore.cpp
+++ b/Engine/Plugins/Editor/ContentBrowser/ContentBrowserFileDataSource/Source/ContentBrowserFileDataSource/Private/ContentBrowserFileDataCore.cpp
@@ -17,19 +17,44 @@
 namespace ContentBrowserFileData
 {
 
-FContentBrowserItemData CreateFolderItem(UContentBrowserDataSource* InOwnerDataSource, const FName InVirtualPath, const FName InInternalPath, const FString& InFilename, TWeakPtr<const ContentBrowserFileData::FDirectoryActions> InDirectoryActions)
+FContentBrowserItemData CreateFolderItem(
+	UContentBrowserDataSource* InOwnerDataSource,
+	const FName InVirtualPath,
+	const FName InInternalPath,
+	const FString& InFilename,
+	TWeakPtr<const ContentBrowserFileData::FDirectoryActions> InDirectoryActions)
 {
 	const FString FolderItemName = FPaths::GetCleanFilename(InFilename);
-	return FContentBrowserItemData(InOwnerDataSource, EContentBrowserItemFlags::Type_Folder | EContentBrowserItemFlags::Category_Misc, InVirtualPath, *FolderItemName, FText(), MakeShared<FContentBrowserFolderItemDataPayload>(InInternalPath, InFilename, MoveTemp(InDirectoryActions)));
+	return FContentBrowserItemData(
+		InOwnerDataSource,
+		EContentBrowserItemFlags::Type_Folder | EContentBrowserItemFlags::Category_Misc,
+		InVirtualPath,
+		*FolderItemName,
+		FText(),
+		MakeShared<FContentBrowserFolderItemDataPayload>(InInternalPath, InFilename, MoveTemp(InDirectoryActions)),
+		{ InInternalPath });
 }
 
-FContentBrowserItemData CreateFileItem(UContentBrowserDataSource* InOwnerDataSource, const FName InVirtualPath, const FName InInternalPath, const FString& InFilename, TWeakPtr<const ContentBrowserFileData::FFileActions> InFileActions)
+FContentBrowserItemData CreateFileItem(
+	UContentBrowserDataSource* InOwnerDataSource,
+	const FName InVirtualPath,
+	const FName InInternalPath,
+	const FString& InFilename,
+	TWeakPtr<const ContentBrowserFileData::FFileActions> InFileActions)
 {
 	const FString FileItemName = FPaths::GetBaseFilename(InFilename);
-	return FContentBrowserItemData(InOwnerDataSource, EContentBrowserItemFlags::Type_File | EContentBrowserItemFlags::Category_Misc, InVirtualPath, *FileItemName, FText(), MakeShared<FContentBrowserFileItemDataPayload>(InInternalPath, InFilename, MoveTemp(InFileActions)));
+	return FContentBrowserItemData(
+		InOwnerDataSource,
+		EContentBrowserItemFlags::Type_File | EContentBrowserItemFlags::Category_Misc,
+		InVirtualPath,
+		*FileItemName,
+		FText(),
+		MakeShared<FContentBrowserFileItemDataPayload>(InInternalPath, InFilename, MoveTemp(InFileActions)),
+		{ InInternalPath });
 }
 
-TSharedPtr<const FContentBrowserFolderItemDataPayload> GetFolderItemPayload(const UContentBrowserDataSource* InOwnerDataSource, const FContentBrowserItemData& InItem)
+TSharedPtr<const FContentBrowserFolderItemDataPayload> GetFolderItemPayload(
+	const UContentBrowserDataSource* InOwnerDataSource, const FContentBrowserItemData& InItem)
 {
 	if (InItem.GetOwnerDataSource() == InOwnerDataSource && InItem.IsFolder())
 	{
diff --git a/Engine/Plugins/Editor/ContentBrowser/ContentBrowserFileDataSource/Source/ContentBrowserFileDataSource/Private/ContentBrowserFileDataSource.cpp b/Engine/Plugins/Editor/ContentBrowser/ContentBrowserFileDataSource/Source/ContentBrowserFileDataSource/Private/ContentBrowserFileDataSource.cpp
--- a/Engine/Plugins/Editor/ContentBrowser/ContentBrowserFileDataSource/Source/ContentBrowserFileDataSource/Private/ContentBrowserFileDataSource.cpp
+++ b/Engine/Plugins/Editor/ContentBrowser/ContentBrowserFileDataSource/Source/ContentBrowserFileDataSource/Private/ContentBrowserFileDataSource.cpp
@@ -1021,18 +1021,20 @@ bool UContentBrowserFileDataSource::CreateFolder(const FName InPath, FContentBro
 
 	FContentBrowserItemData NewItemData(
 		this,
-		EContentBrowserItemFlags::Type_Folder | EContentBrowserItemFlags::Category_Misc | EContentBrowserItemFlags::Temporary_Creation,
+		EContentBrowserItemFlags::Type_Folder | EContentBrowserItemFlags::Category_Misc
+			| EContentBrowserItemFlags::Temporary_Creation,
 		InPath,
 		*FolderItemName,
 		FText::AsCultureInvariant(FolderItemName),
-		MakeShared<FContentBrowserFolderItemDataPayload>(InternalPath, InternalDiskPath, Config.GetDirectoryActions())
-		);
+		MakeShared<FContentBrowserFolderItemDataPayload>(InternalPath, InternalDiskPath, Config.GetDirectoryActions()),
+		{ InternalPath });
 
 	OutPendingItem = FContentBrowserItemDataTemporaryContext(
 		MoveTemp(NewItemData),
-		FContentBrowserItemDataTemporaryContext::FOnValidateItem::CreateUObject(this, &UContentBrowserFileDataSource::OnValidateItemName),
-		FContentBrowserItemDataTemporaryContext::FOnFinalizeItem::CreateUObject(this, &UContentBrowserFileDataSource::OnFinalizeCreateFolder)
-		);
+		FContentBrowserItemDataTemporaryContext::FOnValidateItem::CreateUObject(
+			this, &UContentBrowserFileDataSource::OnValidateItemName),
+		FContentBrowserItemDataTemporaryContext::FOnFinalizeItem::CreateUObject(
+			this, &UContentBrowserFileDataSource::OnFinalizeCreateFolder));
 
 	return true;
 }
@@ -1282,24 +1284,27 @@ bool UContentBrowserFileDataSource::DuplicateItem(const FContentBrowserItemData&
 		checkf(NewItemPayload, TEXT("DuplicateItem returned true, but NewItemPayload was null!"));
 
 		FName VirtualizedPath;
-		TryConvertInternalPathToVirtual(NewItemPayload->GetInternalPath(), VirtualizedPath);
+		FName InternalPath = NewItemPayload->GetInternalPath();
+		TryConvertInternalPathToVirtual(InternalPath, VirtualizedPath);
 
 		const FString FileItemName = FPaths::GetBaseFilename(NewItemPayload->GetFilename());
 
 		FContentBrowserItemData NewItemData(
 			this,
-			EContentBrowserItemFlags::Type_File | EContentBrowserItemFlags::Category_Misc | EContentBrowserItemFlags::Temporary_Duplication,
+			EContentBrowserItemFlags::Type_File | EContentBrowserItemFlags::Category_Misc
+				| EContentBrowserItemFlags::Temporary_Duplication,
 			VirtualizedPath,
 			*FileItemName,
 			FText::AsCultureInvariant(FileItemName),
-			NewItemPayload
-			);
+			NewItemPayload,
+			{ InternalPath });
 
 		OutPendingItem = FContentBrowserItemDataTemporaryContext(
 			MoveTemp(NewItemData),
-			FContentBrowserItemDataTemporaryContext::FOnValidateItem::CreateUObject(this, &UContentBrowserFileDataSource::OnValidateItemName),
-			FContentBrowserItemDataTemporaryContext::FOnFinalizeItem::CreateUObject(this, &UContentBrowserFileDataSource::OnFinalizeDuplicateFile)
-			);
+			FContentBrowserItemDataTemporaryContext::FOnValidateItem::CreateUObject(
+				this, &UContentBrowserFileDataSource::OnValidateItemName),
+			FContentBrowserItemDataTemporaryContext::FOnFinalizeItem::CreateUObject(
+				this, &UContentBrowserFileDataSource::OnFinalizeDuplicateFile));
 
 		return true;
 	}
@@ -1315,11 +1320,19 @@ bool UContentBrowserFileDataSource::BulkDuplicateItems(TArrayView<const FContent
 		for (const TSharedRef<const FContentBrowserFileItemDataPayload>& NewItemPayload : NewItemPayloads)
 		{
 			FName VirtualizedPath;
-			TryConvertInternalPathToVirtual(NewItemPayload->GetInternalPath(), VirtualizedPath);
+			FName InternalPath = NewItemPayload->GetInternalPath();
+			TryConvertInternalPathToVirtual(InternalPath, VirtualizedPath);
 
 			const FString FileItemName = FPaths::GetBaseFilename(NewItemPayload->GetFilename());
 
-			OutNewItems.Emplace(FContentBrowserItemData(this, EContentBrowserItemFlags::Type_File | EContentBrowserItemFlags::Category_Misc, VirtualizedPath, *FileItemName, FText(), NewItemPayload));
+			OutNewItems.Emplace(FContentBrowserItemData(
+				this,
+				EContentBrowserItemFlags::Type_File | EContentBrowserItemFlags::Category_Misc,
+				VirtualizedPath,
+				*FileItemName,
+				FText(),
+				NewItemPayload,
+				{ InternalPath }));
 		}
 
 		return true;
@@ -1997,19 +2010,22 @@ void UContentBrowserFileDataSource::OnNewFileRequested(const FName InDestFolderP
 	const FString NewFileItemName = FPaths::GetBaseFilename(NewFilename);
 
 	FContentBrowserItemData NewItemData(
-		this, 
-		EContentBrowserItemFlags::Type_File | EContentBrowserItemFlags::Category_Misc | EContentBrowserItemFlags::Temporary_Creation,
-		VirtualizedPath, 
+		this,
+		EContentBrowserItemFlags::Type_File | EContentBrowserItemFlags::Category_Misc
+			| EContentBrowserItemFlags::Temporary_Creation,
+		VirtualizedPath,
 		*NewFileItemName,
 		FText::AsCultureInvariant(NewFileItemName),
-		MakeShared<FContentBrowserFileItemDataPayload_Creation>(NewInternalFilePath, NewFilename, InFileActions, MoveTemp(CreationConfig))
-		);
+		MakeShared<FContentBrowserFileItemDataPayload_Creation>(
+			NewInternalFilePath, NewFilename, InFileActions, MoveTemp(CreationConfig)),
+		{ NewInternalFilePath });
 
 	InOnBeginItemCreation.Execute(FContentBrowserItemDataTemporaryContext(
-		MoveTemp(NewItemData), 
-		FContentBrowserItemDataTemporaryContext::FOnValidateItem::CreateUObject(this, &UContentBrowserFileDataSource::OnValidateItemName),
-		FContentBrowserItemDataTemporaryContext::FOnFinalizeItem::CreateUObject(this, &UContentBrowserFileDataSource::OnFinalizeCreateFile)
-		));
+		MoveTemp(NewItemData),
+		FContentBrowserItemDataTemporaryContext::FOnValidateItem::CreateUObject(
+			this, &UContentBrowserFileDataSource::OnValidateItemName),
+		FContentBrowserItemDataTemporaryContext::FOnFinalizeItem::CreateUObject(
+			this, &UContentBrowserFileDataSource::OnFinalizeCreateFile)));
 }
 
 bool UContentBrowserFileDataSource::OnValidateItemName(const FContentBrowserItemData& InItem, const FString& InProposedName, FText* OutErrorMsg)
diff --git a/Engine/Source/Editor/ContentBrowserData/Private/ContentBrowserItemData.cpp b/Engine/Source/Editor/ContentBrowserData/Private/ContentBrowserItemData.cpp
--- a/Engine/Source/Editor/ContentBrowserData/Private/ContentBrowserItemData.cpp
+++ b/Engine/Source/Editor/ContentBrowserData/Private/ContentBrowserItemData.cpp
@@ -9,11 +9,19 @@
 #include "Templates/UnrealTemplate.h"
 #include "UObject/UnrealNames.h"
 
-FContentBrowserItemData::FContentBrowserItemData(UContentBrowserDataSource* InOwnerDataSource, EContentBrowserItemFlags InItemFlags, FName InVirtualPath, FName InItemName, FText InDisplayNameOverride, TSharedPtr<const IContentBrowserItemDataPayload> InPayload)
+FContentBrowserItemData::FContentBrowserItemData(
+	UContentBrowserDataSource* InOwnerDataSource,
+	EContentBrowserItemFlags InItemFlags,
+	FName InVirtualPath,
+	FName InItemName,
+	FText InDisplayNameOverride,
+	TSharedPtr<const IContentBrowserItemDataPayload> InPayload,
+	TOptional<FName> InInternalPath)
 	: OwnerDataSource(InOwnerDataSource)
 	, ItemFlags(InItemFlags)
 	, VirtualPath(InVirtualPath)
 	, ItemName(InItemName)
+	, InternalPath(InInternalPath)
 	, CachedDisplayName(MoveTemp(InDisplayNameOverride))
 	, Payload(MoveTemp(InPayload))
 {
@@ -100,21 +108,20 @@ FName FContentBrowserItemData::GetVirtualPath() const
 
 FName FContentBrowserItemData::GetInvariantPath() const
 {
-	if (UContentBrowserDataSource* DataSource = OwnerDataSource.Get())
+	FName Internal = GetInternalPath();
+	if (!Internal.IsNone())
 	{
-		if (!VirtualPath.IsNone())
-		{
-			FName ConvertedPath;
-			DataSource->TryConvertVirtualPath(VirtualPath, ConvertedPath);
-			return ConvertedPath;
-		}
+		return Internal;
 	}
-
-	return NAME_None;
+	return VirtualPath; // None or a useful path
 }
 
 FName FContentBrowserItemData::GetInternalPath() const
 {
+	if (InternalPath.IsSet())
+	{
+		return InternalPath.GetValue();
+	}
 	if (UContentBrowserDataSource* DataSource = OwnerDataSource.Get())
 	{
 		if (!VirtualPath.IsNone())
@@ -122,6 +129,7 @@ FName FContentBrowserItemData::GetInternalPath() const
 			FName ConvertedPath;
 			if (DataSource->TryConvertVirtualPath(VirtualPath, ConvertedPath) == EContentBrowserPathType::Internal)
 			{
+				InternalPath = ConvertedPath;
 				return ConvertedPath;
 			}
 		}
diff --git a/Engine/Source/Editor/ContentBrowserData/Public/ContentBrowserItemData.h b/Engine/Source/Editor/ContentBrowserData/Public/ContentBrowserItemData.h
--- a/Engine/Source/Editor/ContentBrowserData/Public/ContentBrowserItemData.h
+++ b/Engine/Source/Editor/ContentBrowserData/Public/ContentBrowserItemData.h
@@ -125,14 +125,25 @@ public:
 	/**
 	 * Create an internal Content Browser item.
 	 *
-	 * @param InOwnerDataSource A pointer to the data source that manages the thing represented by this item. This is usually set, but may be null in rare circumstances (ie, creating a dummy placeholder item with no owner).
+	 * @param InOwnerDataSource A pointer to the data source that manages the thing represented by this item. This is
+	 * usually set, but may be null in rare circumstances (ie, creating a dummy placeholder item with no owner).
 	 * @param InItemFlags Flags denoting basic state information for this item instance.
-	 * @param InVirtualPath The complete virtual path that uniquely identifies this item within its owner data source (eg, "/MyRoot/MyFolder/MyFile").
+	 * @param InVirtualPath The complete virtual path that uniquely identifies this item within its owner data source
+	 * (eg, "/MyRoot/MyFolder/MyFile").
 	 * @param InItemName The leaf-name of this item (eg, "MyFile").
-	 * @param InDisplayNameOverride The user-facing name of this item (eg, "MyFile"). This will be lazily set to InItemName if no override is provided. 
+	 * @param InDisplayNameOverride The user-facing name of this item (eg, "MyFile"). This will be lazily set to
+	 * InItemName if no override is provided.
 	 * @param InPayload Any data source defined payload data for this item.
+	 * @param InInternalPath Precalculated internal path if available
 	 */
-	FContentBrowserItemData(UContentBrowserDataSource* InOwnerDataSource, EContentBrowserItemFlags InItemFlags, FName InVirtualPath, FName InItemName, FText InDisplayNameOverride, TSharedPtr<const IContentBrowserItemDataPayload> InPayload);
+	FContentBrowserItemData(
+		UContentBrowserDataSource* InOwnerDataSource,
+		EContentBrowserItemFlags InItemFlags,
+		FName InVirtualPath,
+		FName InItemName,
+		FText InDisplayNameOverride,
+		TSharedPtr<const IContentBrowserItemDataPayload> InPayload,
+		TOptional<FName> InInternalPath = {});
 
 	/**
 	 * Copy support.
@@ -267,6 +278,9 @@ private:
 	/** The leaf-name of this item (eg, "MyFile") */
 	FName ItemName;
 
+	/** Cached/pre-computed value for GetInternalPath */
+	mutable TOptional<FName> InternalPath;
+
 	/** The user-facing name of this item (eg, "MyFile") */
 	mutable FText CachedDisplayName;
 
