Optimizing path permission list checks in ContentBrowserAssetDataSource by optimizing the checks themselves and skipping checks for child paths where possible.
Store paths in permission lists in TDirectoryTree so that checks scale with path depth rather than number of paths.
Add PassesStartsWithFilterRecursive which can return whether a Fail/Pass result for a particular path can be followed by a Pass/Fail result for a child path. (e.g. a path is allowed and there are no explicit denials for any children of that path)
Deprecate GetDenyList and GetAllowList.
Add unit tests for FPathPermissionList.

diff --git a/Engine/Plugins/Editor/ContentBrowser/ContentBrowserAssetDataSource/Source/ContentBrowserAssetDataSource/Private/ContentBrowserAssetDataSource.cpp b/Engine/Plugins/Editor/ContentBrowser/ContentBrowserAssetDataSource/Source/ContentBrowserAssetDataSource/Private/ContentBrowserAssetDataSource.cpp
--- a/Engine/Plugins/Editor/ContentBrowser/ContentBrowserAssetDataSource/Source/ContentBrowserAssetDataSource/Private/ContentBrowserAssetDataSource.cpp
+++ b/Engine/Plugins/Editor/ContentBrowser/ContentBrowserAssetDataSource/Source/ContentBrowserAssetDataSource/Private/ContentBrowserAssetDataSource.cpp
@@ -774,14 +774,15 @@ PRAGMA_ENABLE_DEPRECATION_WARNINGS
 		}
 
 		// Remove any inclusive paths that aren't in the explicit AllowList set
-		if (Params.PathPermissionList && Params.PathPermissionList->GetAllowList().Num() > 0)
+		if (Params.PathPermissionList && Params.PathPermissionList->HasAllowListEntries())
 		{
 			FARCompiledFilter CompiledPathFilterAllowList;
 			{
-				CompiledPathFilterAllowList.PackagePaths.Reserve(Params.PathPermissionList->GetAllowList().Num());
-				for (const auto& AllowListPair : Params.PathPermissionList->GetAllowList())
+				TArray<FString> AllowList = Params.PathPermissionList->GetAllowListEntries();
+				CompiledPathFilterAllowList.PackagePaths.Reserve(AllowList.Num());
+				for (const FString& AllowListEntry : AllowList)
 				{
-					FName PackageName = *AllowListPair.Key;
+					FName PackageName{AllowListEntry};
 					CompiledPathFilterAllowList.PackagePaths.Add(PackageName);
 
 					constexpr bool bIsRecursive = true;
@@ -813,14 +814,15 @@ PRAGMA_ENABLE_DEPRECATION_WARNINGS
 		}
 
 		// Remove any inclusive classes that aren't in the explicit allow list set
-		if (Params.ClassPermissionList && Params.ClassPermissionList->GetAllowList().Num() > 0)
+		if (Params.ClassPermissionList && Params.ClassPermissionList->HasAllowListEntries())
 		{
 			FARCompiledFilter CompiledClassFilterAllowList;
 			{
 				FARFilter AllowListClassFilter;
-				for (const auto& AllowListPair : Params.ClassPermissionList->GetAllowList())
+				TArray<FString> AllowList = Params.ClassPermissionList->GetAllowListEntries();
+				for (const FString& Path : AllowList)
 				{
-					AllowListClassFilter.ClassPaths.Add(FTopLevelAssetPath(AllowListPair.Key));
+					AllowListClassFilter.ClassPaths.Add(FTopLevelAssetPath(Path));
 				}
 				Params.AssetRegistry->CompileFilter(AllowListClassFilter, CompiledClassFilterAllowList);
 			}
@@ -894,12 +896,13 @@ PRAGMA_ENABLE_DEPRECATION_WARNINGS
 		}
 
 		// Add any exclusive paths that are in the explicit DenyList set
-		if (Params.PathPermissionList && Params.PathPermissionList->GetDenyList().Num() > 0)
+		if (Params.PathPermissionList && Params.PathPermissionList->HasDenyListEntries())
 		{
-			CompiledExclusiveFilter.PackagePaths.Reserve(Params.PathPermissionList->GetDenyList().Num());
-			for (const auto& FilterPair : Params.PathPermissionList->GetDenyList())
+			TArray<FString> DenyListEntries = Params.PathPermissionList->GetDenyListEntries();
+			CompiledExclusiveFilter.PackagePaths.Reserve(DenyListEntries.Num());
+			for (const FString& PathString : DenyListEntries)
 			{
-				FName Path = *FilterPair.Key;
+				FName Path{PathString};
 				CompiledExclusiveFilter.PackagePaths.Add(Path);
 
 				constexpr bool bIsRecursive = true;
@@ -913,14 +916,14 @@ PRAGMA_ENABLE_DEPRECATION_WARNINGS
 		}
 
 		// Add any exclusive classes that are in the explicit DenyList set
-		if (Params.ClassPermissionList && Params.ClassPermissionList->GetDenyList().Num() > 0)
+		if (Params.ClassPermissionList && Params.ClassPermissionList->HasDenyListEntries())
 		{
 			FARCompiledFilter CompiledClassFilter;
 			{
 				FARFilter ClassFilter;
-				for (const auto& FilterPair : Params.ClassPermissionList->GetDenyList())
+				for (const FString& Path : Params.ClassPermissionList->GetDenyListEntries())
 				{
-					ClassFilter.ClassPaths.Add(FTopLevelAssetPath(FilterPair.Key));
+					ClassFilter.ClassPaths.Add(FTopLevelAssetPath(Path));
 				}
 				Params.AssetRegistry->CompileFilter(ClassFilter, CompiledClassFilter);
 			}
@@ -1038,10 +1041,11 @@ PRAGMA_ENABLE_DEPRECATION_WARNINGS
 
 							ShowInclusiveFilter.bRecursivePaths = true;
 
-							ShowInclusiveFilter.PackagePaths.Reserve(FolderPermissionList->GetAllowList().Num());
-							for (const auto& AllowListPair : ClassPermissionList->GetAllowList())
+							TArray<FString> AllowList = ClassPermissionList->GetAllowListEntries();
+							ShowInclusiveFilter.PackagePaths.Reserve(AllowList.Num());
+							for (const FString& Path : AllowList)
 							{
-								ShowInclusiveFilter.PackagePaths.Add(*(AllowListPair.Key));
+								ShowInclusiveFilter.PackagePaths.Emplace(Path);
 							}
 
 							Params.AssetRegistry->CompileFilter(ShowInclusiveFilter, CompiledShowInclusiveFilter);
@@ -1084,10 +1088,11 @@ PRAGMA_ENABLE_DEPRECATION_WARNINGS
 
 							ShowExclusiveFilter.bRecursivePaths = true;
 
-							ShowExclusiveFilter.PackagePaths.Reserve(FolderPermissionList->GetDenyList().Num());
-							for (const auto& DenyListPair : ClassPermissionList->GetDenyList())
+							TArray<FString> DenyList = ClassPermissionList->GetDenyListEntries();
+							ShowExclusiveFilter.PackagePaths.Reserve(DenyList.Num());
+							for (const FString& Path : DenyList)
 							{
-								ShowExclusiveFilter.PackagePaths.Add(*(DenyListPair.Key));
+								ShowExclusiveFilter.PackagePaths.Add(FName(Path));
 							}
 
 							Params.AssetRegistry->CompileFilter(ShowExclusiveFilter, CompiledShowExclusiveFilter);
@@ -1107,16 +1112,17 @@ PRAGMA_ENABLE_DEPRECATION_WARNINGS
 
 
 					// Compile the convert if fail inclusive filter
-					if (!ClassPermissionList->GetAllowList().IsEmpty())
+					if (ClassPermissionList->HasAllowListEntries())
 					{
 						FARCompiledFilter CompiledConvertIfFailInclusiveFilter;
 						FARFilter ConvertIfFailInclusiveFilter;
 
 						// Remove any inclusive classes that aren't in the explicit allow list set
-						ConvertIfFailInclusiveFilter.ClassPaths.Reserve(ClassPermissionList->GetAllowList().Num());
-						for (const auto& AllowListPair : ClassPermissionList->GetAllowList())
+						TArray<FString> AllowList = ClassPermissionList->GetAllowListEntries();
+						ConvertIfFailInclusiveFilter.ClassPaths.Reserve(AllowList.Num());
+						for (const FString& Path : AllowList)
 						{
-							ConvertIfFailInclusiveFilter.ClassPaths.Add(FTopLevelAssetPath(AllowListPair.Key));
+							ConvertIfFailInclusiveFilter.ClassPaths.Add(FTopLevelAssetPath(Path));
 						}
 
 						Params.AssetRegistry->CompileFilter(ConvertIfFailInclusiveFilter, CompiledConvertIfFailInclusiveFilter);
@@ -1136,16 +1142,17 @@ PRAGMA_ENABLE_DEPRECATION_WARNINGS
 					}
 
 					// Compile the convert if fail exclusive filter
-					if (!ClassPermissionList->GetDenyList().IsEmpty())
+					if (ClassPermissionList->HasDenyListEntries())
 					{
 						FARCompiledFilter CompiledConvertIfFailExclusiveFilter;
 						FARFilter ConvertIfFailExclusiveFilter;
 
 						// Add any exclusive classes that are in the explicit DenyList set
-						ConvertIfFailExclusiveFilter.ClassPaths.Reserve(ClassPermissionList->GetDenyList().Num());
-						for (const auto& FilterPair : ClassPermissionList->GetDenyList())
+						TArray<FString> DenyList = ClassPermissionList->GetDenyListEntries();
+						ConvertIfFailExclusiveFilter.ClassPaths.Reserve(DenyList.Num());
+						for (const FString& Path : DenyList)
 						{
-							ConvertIfFailExclusiveFilter.ClassPaths.Add(FTopLevelAssetPath(FilterPair.Key));
+							ConvertIfFailExclusiveFilter.ClassPaths.Add(FTopLevelAssetPath(Path));
 						}
 
 
@@ -1334,14 +1341,16 @@ PRAGMA_ENABLE_DEPRECATION_WARNINGS
 		}
 
 		// Remove any inclusive paths that aren't in the explicit AllowList set
-		if (Params.PathPermissionList && Params.PathPermissionList->GetAllowList().Num() > 0)
+		if (Params.PathPermissionList && Params.PathPermissionList->HasAllowListEntries())
 		{
 			FARCompiledFilter CompiledPathFilterAllowList;
 			{
 				FARFilter AllowListPathFilter;
-				for (const auto& AllowListPair : Params.PathPermissionList->GetAllowList())
+				TArray<FString> AllowList = Params.PathPermissionList->GetAllowListEntries();
+				AllowListPathFilter.PackagePaths.Reserve(AllowList.Num());
+				for (const FString& Path : AllowList)
 				{
-					AllowListPathFilter.PackagePaths.Add(*AllowListPair.Key);
+					AllowListPathFilter.PackagePaths.Emplace(Path);
 				}
 				AllowListPathFilter.bRecursivePaths = true;
 				CreateCompiledFilter(AllowListPathFilter, CompiledPathFilterAllowList);
@@ -1365,14 +1374,15 @@ PRAGMA_ENABLE_DEPRECATION_WARNINGS
 		}
 
 		// Remove any inclusive classes that aren't in the explicit allow list set
-		if (Params.ClassPermissionList && Params.ClassPermissionList->GetAllowList().Num() > 0)
+		if (Params.ClassPermissionList && Params.ClassPermissionList->HasAllowListEntries())
 		{
 			FARCompiledFilter CompiledClassFilterAllowList;
 			{
 				FARFilter AllowListClassFilter;
-				for (const auto& AllowListPair : Params.ClassPermissionList->GetAllowList())
+				TArray<FString> AllowList = Params.ClassPermissionList->GetAllowListEntries();
+				for (const FString& Path : AllowList)
 				{
-					AllowListClassFilter.ClassPaths.Add(FTopLevelAssetPath(AllowListPair.Key));
+					AllowListClassFilter.ClassPaths.Add(FTopLevelAssetPath(Path));
 				}
 				Params.AssetRegistry->CompileFilter(AllowListClassFilter, CompiledClassFilterAllowList);
 			}
@@ -1424,14 +1434,14 @@ PRAGMA_ENABLE_DEPRECATION_WARNINGS
 		}
 
 		// Add any exclusive paths that are in the explicit DenyList set
-		if (Params.PathPermissionList && Params.PathPermissionList->GetDenyList().Num() > 0)
+		if (Params.PathPermissionList && Params.PathPermissionList->HasDenyListEntries())
 		{
 			FARCompiledFilter CompiledClassFilter;
 			{
 				FARFilter ClassFilter;
-				for (const auto& FilterPair : Params.PathPermissionList->GetDenyList())
+				for (const FString& Path : Params.PathPermissionList->GetDenyListEntries())
 				{
-					ClassFilter.PackagePaths.Add(*FilterPair.Key);
+					ClassFilter.PackagePaths.Add(FName(Path));
 				}
 				ClassFilter.bRecursivePaths = true;
 				CreateCompiledFilter(ClassFilter, CompiledClassFilter);
@@ -1441,14 +1451,14 @@ PRAGMA_ENABLE_DEPRECATION_WARNINGS
 		}
 
 		// Add any exclusive classes that are in the explicit DenyList set
-		if (Params.ClassPermissionList && Params.ClassPermissionList->GetDenyList().Num() > 0)
+		if (Params.ClassPermissionList && Params.ClassPermissionList->HasDenyListEntries())
 		{
 			FARCompiledFilter CompiledClassFilter;
 			{
 				FARFilter ClassFilter;
-				for (const auto& FilterPair : Params.ClassPermissionList->GetDenyList())
+				for (const FString& Path : Params.ClassPermissionList->GetDenyListEntries())
 				{
-					ClassFilter.ClassPaths.Add(FTopLevelAssetPath(FilterPair.Key));
+					ClassFilter.ClassPaths.Add(FTopLevelAssetPath(Path));
 				}
 				Params.AssetRegistry->CompileFilter(ClassFilter, CompiledClassFilter);
 			}
@@ -1566,10 +1576,11 @@ PRAGMA_ENABLE_DEPRECATION_WARNINGS
 
 							ShowInclusiveFilter.bRecursivePaths = true;
 
-							ShowInclusiveFilter.PackagePaths.Reserve(FolderPermissionList->GetAllowList().Num());
-							for (const auto& AllowListPair : ClassPermissionList->GetAllowList())
+							TArray<FString> AllowList = ClassPermissionList->GetAllowListEntries();
+							ShowInclusiveFilter.PackagePaths.Reserve(AllowList.Num());
+							for (const FString& Path : AllowList)
 							{
-								ShowInclusiveFilter.PackagePaths.Add(*(AllowListPair.Key));
+								ShowInclusiveFilter.PackagePaths.Emplace(Path);
 							}
 
 							Params.AssetRegistry->CompileFilter(ShowInclusiveFilter, CompiledShowInclusiveFilter);
@@ -1612,10 +1623,11 @@ PRAGMA_ENABLE_DEPRECATION_WARNINGS
 
 							ShowExclusiveFilter.bRecursivePaths = true;
 
-							ShowExclusiveFilter.PackagePaths.Reserve(FolderPermissionList->GetDenyList().Num());
-							for (const auto& DenyListPair : ClassPermissionList->GetDenyList())
+							TArray<FString> DenyList = ClassPermissionList->GetDenyListEntries();
+							ShowExclusiveFilter.PackagePaths.Reserve(DenyList.Num());
+							for (const FString& Path : DenyList)
 							{
-								ShowExclusiveFilter.PackagePaths.Add(*(DenyListPair.Key));
+								ShowExclusiveFilter.PackagePaths.Add(FName(Path));
 							}
 
 							Params.AssetRegistry->CompileFilter(ShowExclusiveFilter, CompiledShowExclusiveFilter);
@@ -1635,16 +1647,17 @@ PRAGMA_ENABLE_DEPRECATION_WARNINGS
 
 
 					// Compile the convert if fail inclusive filter
-					if (!ClassPermissionList->GetAllowList().IsEmpty())
+					if (ClassPermissionList->HasAllowListEntries())
 					{
 						FARCompiledFilter CompiledConvertIfFailInclusiveFilter;
 						FARFilter ConvertIfFailInclusiveFilter;
 
 						// Remove any inclusive classes that aren't in the explicit allow list set
-						ConvertIfFailInclusiveFilter.ClassPaths.Reserve(ClassPermissionList->GetAllowList().Num());
-						for (const auto& AllowListPair : ClassPermissionList->GetAllowList())
+						TArray<FString> AllowList = ClassPermissionList->GetAllowListEntries();
+						ConvertIfFailInclusiveFilter.ClassPaths.Reserve(AllowList.Num());
+						for (const FString& Path : AllowList)
 						{
-							ConvertIfFailInclusiveFilter.ClassPaths.Add(FTopLevelAssetPath(AllowListPair.Key));
+							ConvertIfFailInclusiveFilter.ClassPaths.Add(FTopLevelAssetPath(Path));
 						}
 
 						Params.AssetRegistry->CompileFilter(ConvertIfFailInclusiveFilter, CompiledConvertIfFailInclusiveFilter);
@@ -1664,16 +1677,17 @@ PRAGMA_ENABLE_DEPRECATION_WARNINGS
 					}
 
 					// Compile the convert if fail exclusive filter
-					if (!ClassPermissionList->GetDenyList().IsEmpty())
+					if (ClassPermissionList->HasDenyListEntries())
 					{
 						FARCompiledFilter CompiledConvertIfFailExclusiveFilter;
 						FARFilter ConvertIfFailExclusiveFilter;
 
 						// Add any exclusive classes that are in the explicit DenyList set
-						ConvertIfFailExclusiveFilter.ClassPaths.Reserve(ClassPermissionList->GetDenyList().Num());
-						for (const auto& FilterPair : ClassPermissionList->GetDenyList())
+						TArray<FString> DenyList = ClassPermissionList->GetDenyListEntries();
+						ConvertIfFailExclusiveFilter.ClassPaths.Reserve(DenyList.Num());
+						for (const FString& Path : DenyList) 
 						{
-							ConvertIfFailExclusiveFilter.ClassPaths.Add(FTopLevelAssetPath(FilterPair.Key));
+							ConvertIfFailExclusiveFilter.ClassPaths.Add(FTopLevelAssetPath(Path));
 						}
 
 
@@ -1738,30 +1752,136 @@ PRAGMA_ENABLE_DEPRECATION_WARNINGS
 	}
 }
 
+enum class EFolderFilterState
+{
+	None = 0, // Check all filters
+	SkipPathInclude = 0x1,
+	SkipPathExclude = 0x2,
+	SkipPermissionList = 0x3,
+};
+ENUM_CLASS_FLAGS(EFolderFilterState);
+
+// Possible outcomes of the filtering here
+//  Failure - do not visit this path or its children
+//	Success - visit this path and it's children
+//	Additional info for success - whether we need to check any more path filters - or which ones we still need to check
+// return value: success or failure
+// InOutFilterState - bitmask of which filters have passed recursively and can be skipped in future
+bool PathPassesCompiledDataFilterRecursive(const FContentBrowserCompiledAssetDataFilter& InFilter,
+	const FName InInternalPath,
+	EFolderFilterState& InOutFilterState)
+{
+	if (InFilter.ExcludedPackagePaths.Contains(InInternalPath)) // PassesExcludedPathsFilter
+	{
+		return false;
+	}
+
+	FNameBuilder PathStr(InInternalPath);
+	FStringView Path(PathStr);
+	if (!ContentBrowserDataUtils::PathPassesAttributeFilter(Path, 0, InFilter.ItemAttributeFilter))
+	{
+		return false;
+	}
+
+	auto UpdateFilterState = [&InOutFilterState](EPathPermissionPrefixResult Result, EFolderFilterState Flag) -> bool {
+		switch (Result)
+		{
+			case EPathPermissionPrefixResult::Fail:
+			case EPathPermissionPrefixResult::FailRecursive:
+				return false;
+			case EPathPermissionPrefixResult::PassRecursive:
+				InOutFilterState |= Flag;
+				return true;
+			case EPathPermissionPrefixResult::Pass:
+			default:
+				return true;
+		}
+	};
+
+	if (!EnumHasAnyFlags(InOutFilterState, EFolderFilterState::SkipPathInclude))
+	{
+		if (InFilter.bRecursivePackagePathsToInclude)
+		{
+			EPathPermissionPrefixResult IncludeResult =
+				InFilter.PackagePathsToInclude.PassesStartsWithFilterRecursive(Path);
+			if (!UpdateFilterState(IncludeResult, EFolderFilterState::SkipPathInclude))
+			{
+				return false;
+			}
+		}
+		else
+		{
+			if (!InFilter.PackagePathsToInclude.PassesFilter(Path))
+			{
+				return false;
+			}
+			// No info on recursive pass/fail for exact matches, can't update flags
+		}
+	}
+
+	if (!EnumHasAnyFlags(InOutFilterState, EFolderFilterState::SkipPathExclude))
+	{
+		if (InFilter.bRecursivePackagePathsToExclude)
+		{
+			EPathPermissionPrefixResult ExcludeResult =
+				InFilter.PackagePathsToExclude.PassesStartsWithFilterRecursive(Path);
+			if (!UpdateFilterState(ExcludeResult, EFolderFilterState::SkipPathExclude))
+			{
+				return false;
+			}
+		}
+		else
+		{
+			if (!InFilter.PackagePathsToExclude.PassesFilter(Path))
+			{
+				return false;
+			}
+			// No info on recursive pass/fail for exact matches, can't update flags
+		}
+	}
+
+	if (!EnumHasAnyFlags(InOutFilterState, EFolderFilterState::SkipPermissionList))
+	{
+		EPathPermissionPrefixResult PermissionResult =
+			InFilter.PathPermissionList.PassesStartsWithFilterRecursive(Path);
+		if (!UpdateFilterState(PermissionResult, EFolderFilterState::SkipPermissionList))
+		{
+			return false;
+		}
+	}
+
+	return true;
+}
+
 void UContentBrowserAssetDataSource::EnumerateFoldersMatchingFilter(UContentBrowserDataSource* DataSource, const FContentBrowserCompiledAssetDataFilter* AssetDataFilter, TFunctionRef<bool(FContentBrowserItemData&&)> InCallback, FSubPathEnumerationFunc SubPathEnumeration, FCreateFolderItemFunc CreateFolderItem)
 {
 	if (AssetDataFilter->bRunFolderQueryOnDemand)
 	{
-		auto HandleInternalPath = [&DataSource, &InCallback, &AssetDataFilter, &SubPathEnumeration, &CreateFolderItem](const FName InInternalPath)
-		{
-			TArray<FName, TInlineAllocator<16>> PathsToScan;
-			PathsToScan.Add(InInternalPath);
+		auto HandleInternalPath = [&DataSource, &InCallback, &AssetDataFilter, &SubPathEnumeration, &CreateFolderItem](
+									  const FName InInternalPath) {
+			TArray<TPair<FName, EFolderFilterState>, TInlineAllocator<16>> PathsToScan;
+			PathsToScan.Add({ InInternalPath, EFolderFilterState::None });
 			while (PathsToScan.Num() > 0)
 			{
-				const FName PathToScan = PathsToScan.Pop(EAllowShrinking::No);
-				SubPathEnumeration(PathToScan, [&DataSource, &InCallback, &AssetDataFilter, &PathsToScan, &CreateFolderItem](FName SubPath)
-				{
-					if (UContentBrowserAssetDataSource::PathPassesCompiledDataFilter(*AssetDataFilter, SubPath))
-					{
-						if (!InCallback(CreateFolderItem(SubPath)))
+				TPair<FName, EFolderFilterState> PathToScan = PathsToScan.Pop(EAllowShrinking::No);
+				EFolderFilterState ParentFilterState = PathToScan.Value;
+				SubPathEnumeration(
+					PathToScan.Key,
+					[&DataSource, &InCallback, &AssetDataFilter, &PathsToScan, &CreateFolderItem, ParentFilterState](
+						FName SubPath) -> bool {
+						EFolderFilterState FilterState = ParentFilterState;
+						if (PathPassesCompiledDataFilterRecursive(*AssetDataFilter, SubPath, FilterState))
 						{
-							return false;
-						}
+							if (!InCallback(CreateFolderItem(SubPath)))
+							{
+								return false;
+							}
 
-						PathsToScan.Add(SubPath);
-					}
-					return true;
-				}, false);
+							PathsToScan.Add({ SubPath, FilterState });
+						}
+						return true;
+					},
+					false);
 			}
 		};
 
diff --git a/Engine/Source/Developer/AssetTools/Private/AssetTools.cpp b/Engine/Source/Developer/AssetTools/Private/AssetTools.cpp
--- a/Engine/Source/Developer/AssetTools/Private/AssetTools.cpp
+++ b/Engine/Source/Developer/AssetTools/Private/AssetTools.cpp
@@ -513,18 +513,13 @@ namespace UE::AssetTools::Private
 			{
 				TempWritableRootPath = InWritableRootPath;
 
-				if (const FPermissionListOwners* DenyListPtr = WritableFolder->GetDenyList().Find(TempWritableRootPath))
+				if (WritableFolder->ContainsDenyListItem(TempWritableRootPath))
 				{
-					DenyList = *DenyListPtr;
-
-					for (const FName Owner : DenyList)
-					{
-						WritableFolder->RemoveDenyListItem(Owner, TempWritableRootPath);
-					}
+					DenyList = WritableFolder->RemoveDenyListItemAndGetOwners(TempWritableRootPath);
 				}
 
 				// No need to add a permission if the list is empty since everything is permitted (except the deny list)
-				if (!WritableFolder->GetAllowList().IsEmpty())
+				if (WritableFolder->HasAllowListEntries())
 				{
 					WritableFolder->AddAllowListItem(TemporalyPermittedByMigration, TempWritableRootPath);
 				}
diff --git a/Engine/Source/Editor/ContentBrowser/Private/ContentBrowserSingleton.cpp b/Engine/Source/Editor/ContentBrowser/Private/ContentBrowserSingleton.cpp
--- a/Engine/Source/Editor/ContentBrowser/Private/ContentBrowserSingleton.cpp
+++ b/Engine/Source/Editor/ContentBrowser/Private/ContentBrowserSingleton.cpp
@@ -1166,11 +1166,12 @@ void FContentBrowserSingleton::RebuildPrivateContentStateCache()
 	ShowPrivateContentState.CachedVirtualPaths.Reset();
 	ShowPrivateContentState.CachedVirtualPaths = MakeShared<FPathPermissionList>();
 
-	const auto& InvariantAllowList = ShowPrivateContentState.InvariantPaths->GetAllowList();
-	for (const TPair<FString, FPermissionListOwners>& InvariantPathOwnerPair : InvariantAllowList)
+	TArray<FString> InvariantAllowList = ShowPrivateContentState.InvariantPaths->GetAllowListEntries();
+	for (const FString& InvariantPath : InvariantAllowList)
 	{
 		FName VirtualPath;
-		IContentBrowserDataModule::Get().GetSubsystem()->ConvertInternalPathToVirtual(FStringView(InvariantPathOwnerPair.Key), VirtualPath);
+		IContentBrowserDataModule::Get().GetSubsystem()->ConvertInternalPathToVirtual(FStringView(InvariantPath),
+			VirtualPath);
 
 		ShowPrivateContentState.CachedVirtualPaths->AddAllowListItem(TEXT("ContentBrowser"), VirtualPath);
 	}
diff --git a/Engine/Source/Runtime/Core/Private/Misc/PathPermissionList.cpp b/Engine/Source/Runtime/Core/Private/Misc/PathPermissionList.cpp
--- a/Engine/Source/Runtime/Core/Private/Misc/PathPermissionList.cpp
+++ b/Engine/Source/Runtime/Core/Private/Misc/PathPermissionList.cpp
@@ -16,6 +16,20 @@
 #include "UObject/NameTypes.h"
 #include "UObject/UnrealNames.h"
 
+// Deprecated methods
+static TMap<FString, FPermissionListOwners> EmptyDeprecatedList;
+const TMap<FString, FPermissionListOwners>& FPathPermissionList::GetDenyList() const
+{
+	return EmptyDeprecatedList;
+}
+
+const TMap<FString, FPermissionListOwners>& FPathPermissionList::GetAllowList() const
+{
+	return EmptyDeprecatedList;
+}
+
+FPathPermissionList::FPathPermissionList(EPathPermissionListType InType) : ListType(InType) { }
+
 bool FPathPermissionList::PassesFilter(const FStringView Item) const
 {
 	if (DenyListAll.Num() > 0)
@@ -25,16 +39,14 @@ bool FPathPermissionList::PassesFilter(const FStringView Item) const
 
 	VerifyItemMatchesListType(Item);
 
-	if (AllowList.Num() > 0 || DenyList.Num() > 0)
+	if (!AllowTree.IsEmpty() || !DenyTree.IsEmpty())
 	{
-		const uint32 ItemHash = GetTypeHash(Item);
-
-		if (AllowList.Num() > 0 && !AllowList.ContainsByHash(ItemHash, Item))
+		if (!AllowTree.IsEmpty() && !AllowTree.Contains(Item))
 		{
 			return false;
 		}
 
-		if (DenyList.ContainsByHash(ItemHash, Item))
+		if (DenyTree.Contains(Item))
 		{
 			return false;
 		}
@@ -54,81 +66,94 @@ bool FPathPermissionList::PassesFilter(const TCHAR* Item) const
 }
 
 bool FPathPermissionList::PassesStartsWithFilter(const FStringView Item, const bool bAllowParentPaths) const
+{
+	switch (PassesStartsWithFilterRecursive(Item, bAllowParentPaths))
+	{
+		case EPathPermissionPrefixResult::Fail:
+		case EPathPermissionPrefixResult::FailRecursive:
+			return false;
+		case EPathPermissionPrefixResult::Pass:
+		case EPathPermissionPrefixResult::PassRecursive:
+		default:
+			return true;
+	}
+}
+
+bool FPathPermissionList::PassesStartsWithFilter(const FName Item, const bool bAllowParentPaths) const
+{
+	return PassesStartsWithFilter(FNameBuilder(Item), bAllowParentPaths);
+}
+
+bool FPathPermissionList::PassesStartsWithFilter(const TCHAR* Item, const bool bAllowParentPaths) const
+{
+	return PassesStartsWithFilter(FStringView(Item), bAllowParentPaths);
+}
+
+EPathPermissionPrefixResult FPathPermissionList::PassesStartsWithFilterRecursive(const FStringView Item,
+	const bool bAllowParentPaths) const
 {
 	VerifyItemMatchesListType(Item);
 
-	if (AllowList.Num() > 0)
+	if (!HasFiltering())
 	{
-		bool bPassedAllowList = false;
-		for (const auto& Other : AllowList)
-		{
-			if (FPathViews::IsParentPathOf(Other.Key, Item))
-			{
-				bPassedAllowList = true;
-				break;
-			}
-
-			if (bAllowParentPaths)
-			{
-				// If allowing parent paths (eg, when filtering folders), then we must also check if the item has a AllowList child path
-				if (FPathViews::IsParentPathOf(Item, Other.Key))
-				{
-					bPassedAllowList = true;
-					break;
-				}
-			}
-		}
+		return EPathPermissionPrefixResult::PassRecursive;
+	}
 
-		if (!bPassedAllowList)
-		{
-			return false;
-		}
+	if (DenyListAll.Num() > 0)
+	{
+		return EPathPermissionPrefixResult::FailRecursive;
 	}
 
-	if (DenyList.Num() > 0)
+	bool bChildrenMayBeDenied = false;
+	if (!DenyTree.IsEmpty())
 	{
-		for (const auto& Other : DenyList)
+		if (DenyTree.FindClosestValue(Item) != nullptr)
 		{
-			if (FPathViews::IsParentPathOf(Other.Key, Item))
-			{
-				return false;
-			}
+			return EPathPermissionPrefixResult::FailRecursive;
 		}
 	}
 
-	if (DenyListAll.Num() > 0)
+	if (AllowTree.IsEmpty())
 	{
-		return false;
+		// Return value here is dependent on whether child paths might still fail deny lists
+		return DenyTree.ContainsChildPaths(Item) ? EPathPermissionPrefixResult::Pass
+												 : EPathPermissionPrefixResult::PassRecursive;
 	}
 
-	return true;
-}
+	bool bPassedAllowList = AllowTree.FindClosestValue(Item) != nullptr;
+	if (!bPassedAllowList && bAllowParentPaths && AllowTree.ContainsChildPaths(Item))
+	{
+		bPassedAllowList = true;
+	}
+	if (bPassedAllowList)
+	{
+		// If we pass an allow list entry, we might later fail a deny list entry
+		// This logic is also correct if bAllowParent paths is true
+		// 	- child paths of Item will also be parents of an entry in AllowTree
+		//	- child paths of Item may also be present in DenyTree
+		return DenyTree.ContainsChildPaths(Item) ? EPathPermissionPrefixResult::Pass
+												 : EPathPermissionPrefixResult::PassRecursive;
+	}
 
-bool FPathPermissionList::PassesStartsWithFilter(const FName Item, const bool bAllowParentPaths) const
-{
-	return PassesStartsWithFilter(FNameBuilder(Item), bAllowParentPaths);
+	// If we don't match any allow list entries now, check if we might later pass a longer allow list entry
+	// This logic is also correct if bAllowParentPaths is true - if there were a parent path of Item in the tree
+	// it would have matched above and covered future calls with children of Item
+	return AllowTree.ContainsChildPaths(Item) ? EPathPermissionPrefixResult::Fail
+											  : EPathPermissionPrefixResult::FailRecursive;
 }
 
-bool FPathPermissionList::PassesStartsWithFilter(const TCHAR* Item, const bool bAllowParentPaths) const
+bool FPathPermissionList::ContainsDenyListItem(FStringView Item) const
 {
-	return PassesStartsWithFilter(FStringView(Item), bAllowParentPaths);
+	return DenyTree.Contains(Item);
 }
 
 bool FPathPermissionList::AddDenyListItem(const FName OwnerName, const FStringView Item)
 {
 	VerifyItemMatchesListType(Item);
 
-	const uint32 ItemHash = GetTypeHash(Item);
-
-	FPermissionListOwners* Owners = DenyList.FindByHash(ItemHash, Item);
-	const bool bFilterChanged = (Owners == nullptr);
-	if (!Owners)
-	{
-		Owners = &DenyList.AddByHash(ItemHash, FString(Item));
-	}
-
-	Owners->AddUnique(OwnerName);
-	
+	bool bExisted = false;
+	DenyTree.FindOrAdd(Item, &bExisted).AddUnique(OwnerName);
+	bool bFilterChanged = !bExisted;
 	if (bFilterChanged && !bSuppressOnFilterChanged)
 	{
 		OnFilterChanged().Broadcast();
@@ -149,21 +174,18 @@ bool FPathPermissionList::AddDenyListItem(const FName OwnerName, const TCHAR* It
 
 bool FPathPermissionList::RemoveDenyListItem(const FName OwnerName, const FStringView Item)
 {
-	const uint32 ItemHash = GetTypeHash(Item);
-
-	FPermissionListOwners* Owners = DenyList.FindByHash(ItemHash, Item);
+	FPermissionListOwners* Owners = DenyTree.Find(Item);
 	if (Owners && Owners->Remove(OwnerName) == 1)
 	{
 		if (Owners->Num() == 0)
 		{
-			DenyList.RemoveByHash(ItemHash, Item);
-		}
-
-		if (!bSuppressOnFilterChanged)
-		{
-			OnFilterChanged().Broadcast();
+			DenyTree.Remove(Item);
+			if (!bSuppressOnFilterChanged)
+			{
+				OnFilterChanged().Broadcast();
+			}
+			return true;
 		}
-		return true;
 	}
 
 	return false;
@@ -179,21 +201,52 @@ bool FPathPermissionList::RemoveDenyListItem(const FName OwnerName, const TCHAR*
 	return RemoveDenyListItem(OwnerName, FStringView(Item));
 }
 
-bool FPathPermissionList::AddAllowListItem(const FName OwnerName, const FStringView Item)
+bool FPathPermissionList::HasDenyListEntries() const
 {
-	VerifyItemMatchesListType(Item);
+	return DenyTree.IsEmpty() == false;
+}
 
-	const uint32 ItemHash = GetTypeHash(Item);
+TArray<FString> FPathPermissionList::GetDenyListEntries() const
+{
+	TArray<FString> Entries;
+	DenyTree.TryGetChildren({}, Entries, EDirectoryTreeGetFlags::Recursive | EDirectoryTreeGetFlags::ImpliedParent);
+	return MoveTemp(Entries);
+}
 
-	FPermissionListOwners* Owners = AllowList.FindByHash(ItemHash, Item);
-	const bool bFilterChanged = (Owners == nullptr);
-	if (!Owners)
+FPermissionListOwners FPathPermissionList::RemoveDenyListItemAndGetOwners(FStringView Item)
+{
+	if (FPermissionListOwners* Owners = DenyTree.Find(Item))
 	{
-		Owners = &AllowList.AddByHash(ItemHash, FString(Item));
+		FPermissionListOwners RemovedOwners = MoveTemp(*Owners);
+		DenyTree.Remove(Item);
+
+		if (!bSuppressOnFilterChanged)
+		{
+			OnFilterChanged().Broadcast();
+		}
+		return RemovedOwners;
 	}
+	return {};
+}
 
-	Owners->AddUnique(OwnerName);
+bool FPathPermissionList::HasAllowListEntries() const
+{
+	return AllowTree.IsEmpty() == false;
+}
 
+TArray<FString> FPathPermissionList::GetAllowListEntries() const
+{
+	TArray<FString> Entries;
+	AllowTree.TryGetChildren({}, Entries, EDirectoryTreeGetFlags::Recursive | EDirectoryTreeGetFlags::ImpliedParent);
+	return MoveTemp(Entries);
+}
+
+bool FPathPermissionList::AddAllowListItem(const FName OwnerName, const FStringView Item)
+{
+	VerifyItemMatchesListType(Item);
+	bool bExisted = false;
+	AllowTree.FindOrAdd(Item).AddUnique(OwnerName);
+	bool bFilterChanged = !bExisted;
 	if (bFilterChanged && !bSuppressOnFilterChanged)
 	{
 		OnFilterChanged().Broadcast();
@@ -228,21 +281,18 @@ bool FPathPermissionList::AddDenyListAll(const FName OwnerName)
 
 bool FPathPermissionList::RemoveAllowListItem(const FName OwnerName, const FStringView Item)
 {
-	const uint32 ItemHash = GetTypeHash(Item);
-
-	FPermissionListOwners* Owners = AllowList.FindByHash(ItemHash, Item);
+	FPermissionListOwners* Owners = AllowTree.Find(Item);
 	if (Owners && Owners->Remove(OwnerName) == 1)
 	{
 		if (Owners->Num() == 0)
 		{
-			AllowList.RemoveByHash(ItemHash, Item);
-		}
-
-		if (!bSuppressOnFilterChanged)
-		{
-			OnFilterChanged().Broadcast();
+			AllowTree.Remove(Item);
+			if (!bSuppressOnFilterChanged)
+			{
+				OnFilterChanged().Broadcast();
+			}
+			return true;
 		}
-		return true;
 	}
 
 	return false;
@@ -260,24 +310,28 @@ bool FPathPermissionList::RemoveAllowListItem(const FName OwnerName, const TCHAR
 
 bool FPathPermissionList::HasFiltering() const
 {
-	return DenyList.Num() > 0 || AllowList.Num() > 0 || DenyListAll.Num() > 0;
+	return !DenyTree.IsEmpty() || !AllowTree.IsEmpty() || DenyListAll.Num() > 0;
 }
 
 TArray<FName> FPathPermissionList::GetOwnerNames() const
 {
 	TArray<FName> OwnerNames;
 
-	for (const auto& It : DenyList)
+	TArray<FString> Entries;
+	DenyTree.TryGetChildren({}, Entries, EDirectoryTreeGetFlags::Recursive | EDirectoryTreeGetFlags::ImpliedParent);
+	for (const FString& DenyEntry : Entries)
 	{
-		for (const auto& OwnerName : It.Value)
+		for (FName OwnerName : *DenyTree.Find(DenyEntry))
 		{
 			OwnerNames.AddUnique(OwnerName);
 		}
 	}
 
-	for (const auto& It : AllowList)
+	Entries.Reset();
+	AllowTree.TryGetChildren({}, Entries, EDirectoryTreeGetFlags::Recursive | EDirectoryTreeGetFlags::ImpliedParent);
+	for (const FString& AllowEntry : Entries)
 	{
-		for (const auto& OwnerName : It.Value)
+		for (FName OwnerName : *AllowTree.Find(AllowEntry))
 		{
 			OwnerNames.AddUnique(OwnerName);
 		}
@@ -295,22 +349,28 @@ bool FPathPermissionList::UnregisterOwner(const FName OwnerName)
 {
 	bool bFilterChanged = false;
 
-	for (auto It = DenyList.CreateIterator(); It; ++It)
+	TArray<FString> Entries;
+	DenyTree.TryGetChildren({}, Entries, EDirectoryTreeGetFlags::Recursive | EDirectoryTreeGetFlags::ImpliedParent);
+	for (const FString& DenyEntry : Entries)
 	{
-		It->Value.Remove(OwnerName);
-		if (It->Value.Num() == 0)
+		FPermissionListOwners* Owners = DenyTree.Find(DenyEntry);
+		Owners->Remove(OwnerName);
+		if (Owners->Num() == 0)
 		{
-			It.RemoveCurrent();
+			DenyTree.Remove(DenyEntry);
 			bFilterChanged = true;
 		}
 	}
 
-	for (auto It = AllowList.CreateIterator(); It; ++It)
+	Entries.Reset();
+	AllowTree.TryGetChildren({}, Entries, EDirectoryTreeGetFlags::Recursive | EDirectoryTreeGetFlags::ImpliedParent);
+	for (const FString& AllowEntry : Entries)
 	{
-		It->Value.Remove(OwnerName);
-		if (It->Value.Num() == 0)
+		FPermissionListOwners* Owners = AllowTree.Find(AllowEntry);
+		Owners->Remove(OwnerName);
+		if (Owners->Num() == 0)
 		{
-			It.RemoveCurrent();
+			AllowTree.Remove(AllowEntry);
 			bFilterChanged = true;
 		}
 	}
@@ -353,23 +413,31 @@ bool FPathPermissionList::Append(const FPathPermissionList& Other)
 	{
 		TGuardValue<bool> Guard(bSuppressOnFilterChanged, true);
 
-		for (const auto& It : Other.DenyList)
+		TArray<FString> Entries;
+		Other.DenyTree.TryGetChildren({},
+			Entries,
+			EDirectoryTreeGetFlags::Recursive | EDirectoryTreeGetFlags::ImpliedParent);
+		for (const FString& DenyEntry : Entries)
 		{
-			for (const auto& OwnerName : It.Value)
+			for (FName OwnerName : *Other.DenyTree.Find(DenyEntry))
 			{
-				bFilterChanged |= AddDenyListItem(OwnerName, It.Key);
+				bFilterChanged |= AddDenyListItem(OwnerName, DenyEntry);
 			}
 		}
 
-		for (const auto& It : Other.AllowList)
+		Entries.Reset();
+		Other.AllowTree.TryGetChildren({},
+			Entries,
+			EDirectoryTreeGetFlags::Recursive | EDirectoryTreeGetFlags::ImpliedParent);
+		for (const FString& AllowEntry : Entries)
 		{
-			for (const auto& OwnerName : It.Value)
+			for (const auto& OwnerName : *Other.AllowTree.Find(AllowEntry))
 			{
-				bFilterChanged |= AddAllowListItem(OwnerName, It.Key);
+				bFilterChanged |= AddAllowListItem(OwnerName, AllowEntry);
 			}
 		}
 
-		for (const auto& OwnerName : Other.DenyListAll)
+		for (FName OwnerName : Other.DenyListAll)
 		{
 			bFilterChanged |= AddDenyListAll(OwnerName);
 		}
@@ -387,55 +455,65 @@ FPathPermissionList FPathPermissionList::CombinePathFilters(const FPathPermissio
 {
 	FPathPermissionList Result;
 
-	if (IsDenyListAll() || OtherFilter.IsDenyListAll())
-	{
-		Result.AddDenyListAll(NAME_None);
-	}
+	Result.DenyListAll.Append(DenyListAll);
+	Result.DenyListAll.Append(OtherFilter.DenyListAll);
 
-	for (const TPair<FString, FPermissionListOwners>& It : GetDenyList())
+	TArray<FString> Entries;
+	DenyTree.TryGetChildren({}, Entries, EDirectoryTreeGetFlags::Recursive | EDirectoryTreeGetFlags::ImpliedParent);
+	for (const FString& DenyEntry : Entries)
 	{
-		for (const FName& OwnerName : It.Value)
+		for (FName OwnerName : *DenyTree.Find(DenyEntry))
 		{
-			Result.AddDenyListItem(OwnerName, It.Key);
+			Result.AddDenyListItem(OwnerName, DenyEntry);
 		}
 	}
 
-	for (const TPair<FString, FPermissionListOwners>& It : OtherFilter.GetDenyList())
+	Entries.Reset();
+	OtherFilter.DenyTree.TryGetChildren({},
+		Entries,
+		EDirectoryTreeGetFlags::Recursive | EDirectoryTreeGetFlags::ImpliedParent);
+	for (const FString& DenyEntry : Entries)
 	{
-		for (const FName& OwnerName : It.Value)
+		for (const FName& OwnerName : *OtherFilter.DenyTree.Find(DenyEntry))
 		{
-			Result.AddDenyListItem(OwnerName, It.Key);
+			Result.AddDenyListItem(OwnerName, DenyEntry);
 		}
 	}
 
-	if (GetAllowList().Num() > 0 || OtherFilter.GetAllowList().Num() > 0)
+	if (!AllowTree.IsEmpty() || !OtherFilter.AllowTree.IsEmpty())
 	{
-		for (const TPair<FString, FPermissionListOwners>& It : GetAllowList())
+		Entries.Reset();
+		AllowTree.TryGetChildren({},
+			Entries,
+			EDirectoryTreeGetFlags::Recursive | EDirectoryTreeGetFlags::ImpliedParent);
+		for (const FString& AllowEntry : Entries)
 		{
-			const FString& Path = It.Key;
-			if (OtherFilter.PassesStartsWithFilter(Path, true))
+			if (OtherFilter.PassesStartsWithFilter(AllowEntry, true))
 			{
-				for (const FName& OwnerName : It.Value)
+				for (const FName& OwnerName : *AllowTree.Find(AllowEntry))
 				{
-					Result.AddAllowListItem(OwnerName, Path);
+					Result.AddAllowListItem(OwnerName, AllowEntry);
 				}
 			}
 		}
 
-		for (const TPair<FString, FPermissionListOwners>& It : OtherFilter.GetAllowList())
+		Entries.Reset();
+		OtherFilter.AllowTree.TryGetChildren({},
+			Entries,
+			EDirectoryTreeGetFlags::Recursive | EDirectoryTreeGetFlags::ImpliedParent);
+		for (const FString& AllowEntry : Entries)
 		{
-			const FString& Path = It.Key;
-			if (PassesStartsWithFilter(Path, true))
+			if (PassesStartsWithFilter(AllowEntry, true))
 			{
-				for (const FName& OwnerName : It.Value)
+				for (const FName& OwnerName : *OtherFilter.AllowTree.Find(AllowEntry))
 				{
-					Result.AddAllowListItem(OwnerName, Path);
+					Result.AddAllowListItem(OwnerName, AllowEntry);
 				}
 			}
 		}
 
 		// Block everything if none of the AllowList paths passed
-		if (Result.GetAllowList().Num() == 0)
+		if (Result.AllowTree.IsEmpty())
 		{
 			Result.AddDenyListAll(NAME_None);
 		}
@@ -481,8 +559,7 @@ FString FPathPermissionList::ToString() const
 {
 	TStringBuilder<4096> StringBuilder;
 
-	auto SortAndAppendOwners = [&StringBuilder](FPermissionListOwners& Owners)
-	{
+	auto SortAndAppendOwners = [&StringBuilder](FPermissionListOwners Owners) {
 		Owners.Sort(FNameLexicalLess());
 
 		StringBuilder.AppendChar(TCHAR('('));
@@ -505,37 +582,36 @@ FString FPathPermissionList::ToString() const
 	if (!DenyListAll.IsEmpty())
 	{
 		StringBuilder.Append(TEXT("Deny All "));
-		FPermissionListOwners SortedDenyListAll = DenyListAll;
-		SortAndAppendOwners(SortedDenyListAll);
+		SortAndAppendOwners(DenyListAll);
 		StringBuilder.Append(TEXT("\n"));;
 	}
 
-	auto AppendList = [&StringBuilder, &SortAndAppendOwners](const TMap<FString, FPermissionListOwners>& List)
-	{
-		TMap<FString, FPermissionListOwners> SortedList = List;
-		SortedList.KeySort(TLess<FString>());
-		for (TPair<FString, FPermissionListOwners>& ListEntry : SortedList)
+	auto AppendList = [&StringBuilder, &SortAndAppendOwners](const TDirectoryTree<FPermissionListOwners>& Tree) {
+		TArray<FString> AllPaths;
+		Tree.TryGetChildren({}, AllPaths, EDirectoryTreeGetFlags::Recursive | EDirectoryTreeGetFlags::ImpliedParent);
+		Algo::Sort(AllPaths);
+		for (const FString& Path : AllPaths)
 		{
 			StringBuilder.AppendChar(TCHAR('\t'));
 			StringBuilder.AppendChar(TCHAR('"'));
-			StringBuilder.Append(ListEntry.Key);
+			StringBuilder.Append(Path);
 			StringBuilder.AppendChar(TCHAR('"'));
 			StringBuilder.AppendChar(TCHAR(' '));
-			SortAndAppendOwners(ListEntry.Value);
+			SortAndAppendOwners(*Tree.Find(Path));
 			StringBuilder.AppendChar(TCHAR('\n'));
 		}
 	};
 
-	if (!DenyList.IsEmpty())
+	if (!DenyTree.IsEmpty())
 	{
 		StringBuilder.Append(TEXT("Deny List\n"));
-		AppendList(DenyList);
+		AppendList(DenyTree);
 	}
 
-	if (!AllowList.IsEmpty())
+	if (!AllowTree.IsEmpty())
 	{
 		StringBuilder.Append(TEXT("Allow List\n"));
-		AppendList(AllowList);
+		AppendList(AllowTree);
 	}
 
 	return StringBuilder.ToString();
diff --git a/Engine/Source/Runtime/Core/Public/Misc/NamePermissionList.h b/Engine/Source/Runtime/Core/Public/Misc/NamePermissionList.h
--- a/Engine/Source/Runtime/Core/Public/Misc/NamePermissionList.h
+++ b/Engine/Source/Runtime/Core/Public/Misc/NamePermissionList.h
@@ -3,6 +3,7 @@
 #pragma once
 
 #include "Containers/Array.h"
+#include "Containers/DirectoryTree.h"
 #include "Containers/Map.h"
 #include "Containers/StringFwd.h"
 #include "Containers/StringView.h"
@@ -122,15 +123,52 @@ enum class EPathPermissionListType
 	ClassPaths	// Class permission list
 };
 
+// Result of non-exact filtering on path prefixes
+enum class EPathPermissionPrefixResult
+{
+	// The query against the list failed because of one of the following:
+	// 	* There was an explicit allow-list and none of the entries was a parent path of the query path
+	//		(or a child of the query path when bAllowParentPaths = true)
+	//	* There were deny-list entries and one of the entries was a parent path of the query path
+	Fail,
+	// The query against the list failed and either:
+	//	* All paths are denied, so queries for child paths will also fail
+	//	* The query failed on a deny-list entry, so child paths will fail on this same entry
+	//	* There is an explicit allow-list and there are no entries with more components than the query path,
+	//	  so children can never pass.
+	FailRecursive,
+	// The query against the list succeeded, but queries for child paths may fail because there are longer
+	// paths in the deny-list with more components
+	Pass,
+	// The query against the list succeeded and queries for child paths will all succeed no deny-list entry can
+	// possibly fail them
+	PassRecursive,
+};
+
+/**
+ * Set of paths that are allowd and/or denied for certain use cases.
+ * A permission list may contain
+ * 	- Blanket denial
+ * 	- Specifically denied paths
+ * 	- Specifically allowed paths
+ * In decreasing order of priority. When performing prefix checks, if a a path matches a denied path, it cannot be
+ * allowed again by a more specific allowed path. If any paths are specifically allowed, paths which do NOT match
+ * something in the allow list are implicitly denied.
+ */
 class FPathPermissionList : public TSharedFromThis<FPathPermissionList>
 {
 public:
-	FPathPermissionList(EPathPermissionListType InType = EPathPermissionListType::Default) 
-		: ListType(InType)
-	{
-	}
-	virtual ~FPathPermissionList() {}
-	
+	CORE_API FPathPermissionList(EPathPermissionListType InType = EPathPermissionListType::Default);
+	PRAGMA_DISABLE_DEPRECATION_WARNINGS
+	virtual ~FPathPermissionList() { }
+	PRAGMA_ENABLE_DEPRECATION_WARNINGS
+
+	FPathPermissionList(const FPathPermissionList&) = default;
+	FPathPermissionList& operator=(const FPathPermissionList&) = default;
+
+	FPathPermissionList(FPathPermissionList&&) = default;
+	FPathPermissionList& operator=(FPathPermissionList&&) = default;
+
 	/** Returns true if passes filter restrictions using exact match */
 	CORE_API bool PassesFilter(const FStringView Item) const;
 
@@ -149,6 +187,46 @@ public:
 	/** Returns true if passes filter restrictions for path */
 	CORE_API bool PassesStartsWithFilter(const TCHAR* Item, const bool bAllowParentPaths = false) const;
 
+	/**
+	 * Checks the given path against the restrictions and return whether it's possible for any child paths to succeed or
+	 * fail as well. Returning PassRecursive or FailRecursive guarantees that no child paths of the queried path can
+	 * fail or pass the filter respectively. Returning Pass or Fail does not guarantee that there is some path which
+	 * fails or passes the filter respectively.
+	 *
+	 * Examples:
+	 *
+	 * Given no deny or allow lists:
+	 * Inputs:
+	 * 	/ -> PassRecursive, because no paths can fail to match the allow list or match the deny list.
+	 *
+	 * Given a deny-list entry:
+	 *  Allow: empty
+	 *  Deny: /Secret
+	 * Inputs:
+	 * 	/ -> Pass, because some children of this path may be denied.
+	 *  /Secret -> FailRecursive, because this path is denied and all children will also be denied.
+	 *  /Public -> PassRecursive, because this path is not demied and no children can be denied.
+	 *
+	 * Given an allow-list entry:
+	 * 	Allow: /JustThis
+	 *  Deny: empty
+	 * Inputs:
+	 *  / -> Fail
+	 * 	/JustThis -> PassRecursive
+	 *  /SomethingElse -> FailRecursive
+	 *
+	 * Given both allow and deny-lists:
+	 *  Allow: /Stuff
+	 * 	Deny: /Stuff/Secret
+	 * Inputs:
+	 * 	/ -> Fail
+	 *  /Stuff -> Pass
+	 * 	/Stuff/Secret -> Fail
+	 * 	/Stuff/Public -> PassRecursive
+	 */
+	CORE_API EPathPermissionPrefixResult PassesStartsWithFilterRecursive(const FStringView Item,
+		const bool bAllowParentPaths = false) const;
+
 	/** 
 	 * Add item to DenyList, this specific item will be filtered out.
 	 * @return whether the filters changed.
@@ -167,6 +245,15 @@ public:
 	 */
 	CORE_API bool AddDenyListItem(const FName OwnerName, const TCHAR* Item);
 
+	/** Returns whether the given path has been denied explicitly with a call to AddDenyListItem. */
+	CORE_API bool ContainsDenyListItem(FStringView Item) const;
+
+	/** Returns whether this list has any explicitly denied paths. */
+	CORE_API bool HasDenyListEntries() const;
+
+	/** Get a copy of the paths explicity denied in this list. */
+	CORE_API TArray<FString> GetDenyListEntries() const;
+
 	/**
 	* Remove item from the DenyList
 	* @return whether the filters changed
@@ -185,6 +272,12 @@ public:
 	*/
 	CORE_API bool RemoveDenyListItem(const FName OwnerName, const TCHAR* Item);
 
+	/**
+	 * Removes an item from the deny list and returns a list of all the owners of that item
+	 * so that the item can be re-introduced.
+	 */
+	CORE_API FPermissionListOwners RemoveDenyListItemAndGetOwners(FStringView Item);
+
 	/**
 	 * Add item to allowlist after which all items not in the allowlist will be filtered out.
 	 * @return whether the filters changed.
@@ -203,6 +296,13 @@ public:
 	 */
 	CORE_API bool AddAllowListItem(const FName OwnerName, const TCHAR* Item);
 
+	/** Returns whether this list has any explicitly allowed paths, which will lead to it denying access to any paths
+	 * not listed.*/
+	CORE_API bool HasAllowListEntries() const;
+
+	/** Returns a copy of the paths explicity allowed in this list */
+	CORE_API TArray<FString> GetAllowListEntries() const;
+
 	/**
 	* Remove item from the AllowList
 	* @return whether the filters changed
@@ -266,10 +366,12 @@ public:
 	CORE_API bool UnregisterOwnersAndAppend(const TArray<FName>& OwnerNamesToRemove, const FPathPermissionList& FiltersToAdd);
 
 	/** Get raw DenyList */
-	const TMap<FString, FPermissionListOwners>& GetDenyList() const { return DenyList; }
-	
+	UE_DEPRECATED(5.5, "GetDenyList is deprecated. Use GetDenyListEntries instead.")
+	const TMap<FString, FPermissionListOwners>& GetDenyList() const;
+
 	/** Get raw allowlist */
-	const TMap<FString, FPermissionListOwners>& GetAllowList() const { return AllowList; }
+	UE_DEPRECATED(5.5, "GetAllowList is deprecated. Use GetAllowListEntries instead.")
+	const TMap<FString, FPermissionListOwners>& GetAllowList() const;
 
 	/** Are all items set to be filtered out */
 	bool IsDenyListAll() const { return DenyListAll.Num() > 0; }
@@ -288,11 +390,11 @@ protected:
 	 */
 	CORE_API void VerifyItemMatchesListType(const FStringView Item) const;
 
-	/** List if items to filter out */
-	TMap<FString, FPermissionListOwners> DenyList;
+	/** Compiled path tree produced from DenyList */
+	TDirectoryTree<FPermissionListOwners> DenyTree;
 
-	/** List of items to allow, if not empty all items will be filtered out unless they are in the list */
-	TMap<FString, FPermissionListOwners> AllowList;
+	/** Compiled path tree produced from AllowList */
+	TDirectoryTree<FPermissionListOwners> AllowTree;
 
 	/** List of owner names that requested all items to be filtered out */
 	FPermissionListOwners DenyListAll;
diff --git a/Engine/Source/Runtime/Core/Tests/Misc/PathPermissionListTests.cpp b/Engine/Source/Runtime/Core/Tests/Misc/PathPermissionListTests.cpp
new file mode 100644
--- /dev/null
+++ b/Engine/Source/Runtime/Core/Tests/Misc/PathPermissionListTests.cpp
@@ -0,0 +1,273 @@
+// Copyright Epic Games, Inc. All Rights Reserved.
+
+
+#if WITH_TESTS
+
+#include "Misc/NamePermissionList.h"
+
+#include "Tests/TestHarnessAdapter.h"
+
+std::ostream& operator<<(std::ostream& Stream, EPathPermissionPrefixResult Value)
+{
+    switch(Value)
+    {
+        case EPathPermissionPrefixResult::Fail:
+            return Stream << "Fail";
+        case EPathPermissionPrefixResult::FailRecursive:
+            return Stream << "FailRecursive";
+        case EPathPermissionPrefixResult::Pass:
+            return Stream << "Pass";
+        case EPathPermissionPrefixResult::PassRecursive:
+            return Stream << "PassRecursive";
+    }
+    return Stream << "Unknown";
+}
+
+TEST_CASE_NAMED(FPathPermissionListExactMatchTests, "System::Core::Misc::PathPermissionList::ExactMatch", "[Core][Misc][PathPermissionList]")
+{
+    SECTION("Deny all")
+    {
+        FPathPermissionList List;
+        List.AddDenyListAll("DenyAll");
+        
+        // Path is denied when everything is denied
+        CHECK_FALSE(List.PassesFilter(TEXTVIEW("/Game/Folder")));
+        
+        // Deny all takes precendence over specific allo
+        List.AddAllowListItem("AllowSpecific", TEXTVIEW("/Game/Folder"));
+        CHECK_FALSE(List.PassesFilter(TEXTVIEW("/Game/Folder")));
+        
+        // Removing deny all allows the path again 
+        List.UnregisterOwner("DenyAll");
+        CHECK(List.PassesFilter(TEXTVIEW("/Game/Folder")));
+    }
+    
+    SECTION("Deny list")
+    {
+        FPathPermissionList List;
+        
+        List.AddDenyListItem("Deny", TEXTVIEW("/Game/Secret"));
+        List.AddDenyListItem("Deny", TEXTVIEW("/Plugin1"));
+        
+        CHECK(List.PassesFilter(TEXTVIEW("/")));
+        CHECK(List.PassesFilter(TEXTVIEW("/Game")));
+        CHECK(List.PassesFilter(TEXTVIEW("/Game/Secret/ActuallyVisible")));
+        CHECK(List.PassesFilter(TEXTVIEW("/Plugin1/AlsoVisible")));
+        
+        CHECK_FALSE(List.PassesFilter(TEXTVIEW("/Plugin1")));
+        CHECK_FALSE(List.PassesFilter(TEXTVIEW("/Game/Secret")));
+    }
+    
+    SECTION("Allow list")
+    {
+        FPathPermissionList List;
+        
+        List.AddAllowListItem("Allow", TEXTVIEW("/Game"));
+        List.AddAllowListItem("Allow", TEXTVIEW("/Game/Public"));
+        List.AddAllowListItem("Allow", TEXTVIEW("/Plugin1"));
+        
+        CHECK(List.PassesFilter(TEXTVIEW("/Game")));
+        CHECK(List.PassesFilter(TEXTVIEW("/Game/Public")));
+        CHECK(List.PassesFilter(TEXTVIEW("/Plugin1")));
+
+        CHECK_FALSE(List.PassesFilter(TEXTVIEW("/")));
+        CHECK_FALSE(List.PassesFilter(TEXTVIEW("/Game/Secret")));
+        CHECK_FALSE(List.PassesFilter(TEXTVIEW("/Plugin1/InPlugin")));
+        CHECK_FALSE(List.PassesFilter(TEXTVIEW("/Plugin2")));
+    }
+
+    SECTION("Mixed allow and deny list")
+    {
+        FPathPermissionList List;
+        
+        List.AddAllowListItem("Allow", TEXTVIEW("/Game"));
+        List.AddAllowListItem("Allow", TEXTVIEW("/Game/Maps"));
+        List.AddAllowListItem("Allow", TEXTVIEW("/Plugin1"));
+        List.AddDenyListItem("Deny", TEXTVIEW("/Plugin1"));
+        List.AddDenyListItem("Deny", TEXTVIEW("/Game"));
+        
+        CHECK(List.PassesFilter(TEXTVIEW("/Game/Maps")));
+
+        CHECK_FALSE(List.PassesFilter(TEXTVIEW("/")));
+        CHECK_FALSE(List.PassesFilter(TEXTVIEW("/Game")));
+        CHECK_FALSE(List.PassesFilter(TEXTVIEW("/Game/Characters")));
+        CHECK_FALSE(List.PassesFilter(TEXTVIEW("/Plugin1")));
+        CHECK_FALSE(List.PassesFilter(TEXTVIEW("/Plugin2")));
+    }
+}
+
+TEST_CASE_NAMED(FPathPermissionListCombineTest, "System::Core::Misc::PathPermissionList::Combine", "[Core][Misc][PathPermissionList]")
+{
+    SECTION("One deny all")
+    {
+        FPathPermissionList First;
+        First.AddDenyListAll("Deny");
+        FPathPermissionList Second;
+        Second.AddAllowListItem("Allow", TEXTVIEW("/Game"));
+        Second.AddAllowListItem("Allow", TEXTVIEW("/Game/Maps"));
+
+        FPathPermissionList Combined = First.CombinePathFilters(Second);
+        
+        CHECK_FALSE(Combined.PassesFilter(TEXTVIEW("/Game")));
+        CHECK_FALSE(Combined.PassesFilter(TEXTVIEW("/Game/Maps")));
+        CHECK_FALSE(Combined.PassesStartsWithFilter(TEXTVIEW("/Game")));
+        CHECK_FALSE(Combined.PassesStartsWithFilter(TEXTVIEW("/Game/Maps")));
+        CHECK_FALSE(Combined.PassesStartsWithFilter(TEXTVIEW("/Game/Characters")));
+        CHECK_FALSE(Combined.PassesStartsWithFilter(TEXTVIEW("/Plugin1")));
+    }   
+    
+    SECTION("Parent path denied")
+    {
+        FPathPermissionList First;
+        First.AddDenyListItem("Deny", "/Game/Maps");
+        FPathPermissionList Second;
+        Second.AddAllowListItem("Allow", TEXTVIEW("/Game"));
+        Second.AddAllowListItem("Allow", TEXTVIEW("/Game/Maps/Desert"));
+
+        FPathPermissionList Combined = First.CombinePathFilters(Second);
+        CHECK(Combined.PassesFilter(TEXTVIEW("/Game")));
+        CHECK_FALSE(Combined.PassesFilter(TEXTVIEW("/Game/Maps")));
+        // "StartsWith" is used when combining to remove /Game/Maps/Desert from final allow list
+        CHECK_FALSE(Combined.PassesFilter(TEXTVIEW("/Game/Maps/Desert"))); 
+        
+        CHECK(Combined.PassesStartsWithFilter(TEXTVIEW("/Game")));
+        CHECK(Combined.PassesStartsWithFilter(TEXTVIEW("/Game/Characters")));
+
+        CHECK_FALSE(Combined.PassesStartsWithFilter(TEXTVIEW("/Plugin1")));
+        CHECK_FALSE(Combined.PassesStartsWithFilter(TEXTVIEW("/Game/Maps")));
+        CHECK_FALSE(Combined.PassesStartsWithFilter(TEXTVIEW("/Game/Maps/Desert")));
+        CHECK_FALSE(Combined.PassesStartsWithFilter(TEXTVIEW("/Game/Maps/Jungle")));
+    }
+}
+
+TEST_CASE_NAMED(FPathPermissionListStartsWithTests, "System::Core::Misc::PathPermissionList::StartsWith", "[Core][Misc][PathPermissionList]")
+{
+    SECTION("Deny all")
+    {
+        FPathPermissionList List;
+        List.AddDenyListAll("DenyAll");
+        
+        CHECK_FALSE(List.PassesStartsWithFilter(TEXTVIEW("/Game")));
+    }
+
+    SECTION("Deny only")
+    {
+        FPathPermissionList List;
+        List.AddDenyListItem("Deny", TEXTVIEW("/Game/Secret"));
+        List.AddDenyListItem("Deny", TEXTVIEW("/Game/Maps/Secret"));
+        List.AddDenyListItem("Deny", TEXTVIEW("/Plugin1/Secret"));
+        List.AddDenyListItem("Deny", TEXTVIEW("/Plugin2"));
+        
+        // Some children of /Game are blocked
+        CHECK(List.PassesStartsWithFilter(TEXTVIEW("/Game")));
+        CHECK_FALSE(List.PassesStartsWithFilter(TEXTVIEW("/Game/Secret")));
+        CHECK(List.PassesStartsWithFilter(TEXTVIEW("/Game/Characters")));
+        CHECK(List.PassesStartsWithFilter(TEXTVIEW("/Game/Maps")));
+        CHECK_FALSE(List.PassesStartsWithFilter(TEXTVIEW("/Game/Maps/Secret")));
+        CHECK(List.PassesStartsWithFilter(TEXTVIEW("/Game/Maps/Desert")));
+        CHECK(List.PassesStartsWithFilter(TEXTVIEW("/Plugin1")));
+        CHECK_FALSE(List.PassesStartsWithFilter(TEXTVIEW("/Plugin1/Secret")));
+        CHECK_FALSE(List.PassesStartsWithFilter(TEXTVIEW("/Plugin2")));
+    }   
+    
+    SECTION("Allow only")
+    {
+        FPathPermissionList List;
+        List.AddAllowListItem("Allow", TEXTVIEW("/Game/Maps"));
+        List.AddAllowListItem("Allow", TEXTVIEW("/Plugin1"));
+        
+        CHECK_FALSE(List.PassesStartsWithFilter(TEXTVIEW("/Game")));
+        CHECK(List.PassesStartsWithFilter(TEXTVIEW("/Game/Maps")));
+        CHECK(List.PassesStartsWithFilter(TEXTVIEW("/Plugin1")));
+        CHECK(List.PassesStartsWithFilter(TEXTVIEW("/Plugin1/Characters")));
+        CHECK_FALSE(List.PassesStartsWithFilter(TEXTVIEW("/Plugin2")));
+    }
+
+    SECTION("Mixed allow and deny")
+    {
+        FPathPermissionList List;
+        List.AddAllowListItem("Allow", TEXTVIEW("/Game"));        
+        List.AddAllowListItem("Allow", TEXTVIEW("/Plugin1"));
+        List.AddAllowListItem("Allow", TEXTVIEW("/Plugin2/Public"));
+        List.AddDenyListItem("Deny", TEXTVIEW("/Game/Secret"));
+        List.AddDenyListItem("Deny", TEXTVIEW("/Game/Maps/Secret"));
+        
+        CHECK(List.PassesStartsWithFilter(TEXTVIEW("/Game")));
+        CHECK_FALSE(List.PassesStartsWithFilter(TEXTVIEW("/Game/Secret")));
+        CHECK(List.PassesStartsWithFilter(TEXTVIEW("/Game/Maps")));
+        CHECK_FALSE(List.PassesStartsWithFilter(TEXTVIEW("/Game/Maps/Secret")));
+
+        CHECK(List.PassesStartsWithFilter(TEXTVIEW("/Plugin1")));
+
+        CHECK_FALSE(List.PassesStartsWithFilter(TEXTVIEW("/Plugin2")));
+        CHECK(List.PassesStartsWithFilter(TEXTVIEW("/Plugin2/Public")));
+        CHECK_FALSE(List.PassesStartsWithFilter(TEXTVIEW("/Plugin2/Private")));
+    }
+}
+
+TEST_CASE_NAMED(FPathPermissionListStartsWithRecursiveTests, "System::Core::Misc::PathPermissionList::StartsWithRecursive", "[Core][Misc][PathPermissionList]")
+{
+    SECTION("Deny all")
+    {
+        FPathPermissionList List;
+        List.AddDenyListAll("DenyAll");
+        
+        CHECK(List.PassesStartsWithFilterRecursive(TEXTVIEW("/Game")) == EPathPermissionPrefixResult::FailRecursive);
+    }
+
+    SECTION("Deny only")
+    {
+        FPathPermissionList List;
+        List.AddDenyListItem("Deny", TEXTVIEW("/Game/Secret"));
+        List.AddDenyListItem("Deny", TEXTVIEW("/Game/Maps/Secret"));
+        List.AddDenyListItem("Deny", TEXTVIEW("/Plugin1/Secret"));
+        List.AddDenyListItem("Deny", TEXTVIEW("/Plugin2"));
+        
+        // Some children of /Game are blocked
+        CHECK(List.PassesStartsWithFilterRecursive(TEXTVIEW("/Game")) == EPathPermissionPrefixResult::Pass);
+        CHECK(List.PassesStartsWithFilterRecursive(TEXTVIEW("/Game/Secret")) == EPathPermissionPrefixResult::FailRecursive);
+        CHECK(List.PassesStartsWithFilterRecursive(TEXTVIEW("/Game/Characters")) == EPathPermissionPrefixResult::PassRecursive);
+        CHECK(List.PassesStartsWithFilterRecursive(TEXTVIEW("/Game/Maps")) == EPathPermissionPrefixResult::Pass);
+        CHECK(List.PassesStartsWithFilterRecursive(TEXTVIEW("/Game/Maps/Secret")) == EPathPermissionPrefixResult::FailRecursive);
+        CHECK(List.PassesStartsWithFilterRecursive(TEXTVIEW("/Game/Maps/Desert")) == EPathPermissionPrefixResult::PassRecursive);
+        CHECK(List.PassesStartsWithFilterRecursive(TEXTVIEW("/Plugin1")) == EPathPermissionPrefixResult::Pass);
+        CHECK(List.PassesStartsWithFilterRecursive(TEXTVIEW("/Plugin1/Secret")) == EPathPermissionPrefixResult::FailRecursive);
+        CHECK(List.PassesStartsWithFilterRecursive(TEXTVIEW("/Plugin2")) == EPathPermissionPrefixResult::FailRecursive);
+    }   
+    
+    SECTION("Allow only")
+    {
+        FPathPermissionList List;
+        List.AddAllowListItem("Allow", TEXTVIEW("/Game/Maps"));
+        List.AddAllowListItem("Allow", TEXTVIEW("/Plugin1"));
+        
+        CHECK(List.PassesStartsWithFilterRecursive(TEXTVIEW("/Game")) == EPathPermissionPrefixResult::Fail);
+        CHECK(List.PassesStartsWithFilterRecursive(TEXTVIEW("/Game/Maps")) == EPathPermissionPrefixResult::PassRecursive);
+        CHECK(List.PassesStartsWithFilterRecursive(TEXTVIEW("/Plugin1")) == EPathPermissionPrefixResult::PassRecursive);
+        CHECK(List.PassesStartsWithFilterRecursive(TEXTVIEW("/Plugin1/Characters")) == EPathPermissionPrefixResult::PassRecursive);
+        CHECK(List.PassesStartsWithFilterRecursive(TEXTVIEW("/Plugin2")) == EPathPermissionPrefixResult::FailRecursive);
+    }
+
+    SECTION("Mixed allow and deny")
+    {
+        FPathPermissionList List;
+        List.AddAllowListItem("Allow", TEXTVIEW("/Game"));        
+        List.AddAllowListItem("Allow", TEXTVIEW("/Plugin1"));
+        List.AddAllowListItem("Allow", TEXTVIEW("/Plugin2/Public"));
+        List.AddDenyListItem("Deny", TEXTVIEW("/Game/Secret"));
+        List.AddDenyListItem("Deny", TEXTVIEW("/Game/Maps/Secret"));
+        
+        CHECK(List.PassesStartsWithFilterRecursive(TEXTVIEW("/Game")) == EPathPermissionPrefixResult::Pass);
+        CHECK(List.PassesStartsWithFilterRecursive(TEXTVIEW("/Game/Secret")) == EPathPermissionPrefixResult::FailRecursive);
+        CHECK(List.PassesStartsWithFilterRecursive(TEXTVIEW("/Game/Maps")) == EPathPermissionPrefixResult::Pass);
+        CHECK(List.PassesStartsWithFilterRecursive(TEXTVIEW("/Game/Maps/Secret")) == EPathPermissionPrefixResult::FailRecursive);
+
+        CHECK(List.PassesStartsWithFilterRecursive(TEXTVIEW("/Plugin1")) == EPathPermissionPrefixResult::PassRecursive);
+
+        CHECK(List.PassesStartsWithFilterRecursive(TEXTVIEW("/Plugin2")) == EPathPermissionPrefixResult::Fail);
+        CHECK(List.PassesStartsWithFilterRecursive(TEXTVIEW("/Plugin2/Public")) == EPathPermissionPrefixResult::PassRecursive);
+        CHECK(List.PassesStartsWithFilterRecursive(TEXTVIEW("/Plugin2/Private")) == EPathPermissionPrefixResult::FailRecursive);
+    }
+}
+
+#endif // WITH_TESTS
\ No newline at end of file
