Fix an issue where some folders might be missing from the content browser if some external content is also mounted under the same mount point.

diff --git a/Engine/Plugins/Editor/ContentBrowser/ContentBrowserAssetDataSource/Source/ContentBrowserAssetDataSource/Private/ContentBrowserAssetDataSource.cpp b/Engine/Plugins/Editor/ContentBrowser/ContentBrowserAssetDataSource/Source/ContentBrowserAssetDataSource/Private/ContentBrowserAssetDataSource.cpp
--- a/Engine/Plugins/Editor/ContentBrowser/ContentBrowserAssetDataSource/Source/ContentBrowserAssetDataSource/Private/ContentBrowserAssetDataSource.cpp
+++ b/Engine/Plugins/Editor/ContentBrowser/ContentBrowserAssetDataSource/Source/ContentBrowserAssetDataSource/Private/ContentBrowserAssetDataSource.cpp
@@ -2929,7 +2929,7 @@ bool UContentBrowserAssetDataSource::IsKnownContentPath(const FName InPackagePat
 		const FCharacterNode* NextNode = NextNodePair->Key.Get();
 
 		// Is the next node terminal
-		if (NextNode->NextNodes.IsEmpty())
+		if (NextNode->bIsEndOfAMountPoint)
 		{
 			// The package path start with a root content path
 			return true;
@@ -3671,12 +3671,6 @@ FContentBrowserItemData UContentBrowserAssetDataSource::OnFinalizeDuplicateAsset
 
 void UContentBrowserAssetDataSource::AddRootContentPathToStateMachine(const FString& InAssetPath)
 {
-	/**
-	 * No need to mark the nodes with a terminal attribute on the last node since they always finish with a '/'.
-	 * They never contains more then two '/'. One at the start and the other at end.
-	 */
-	 ensure(InAssetPath[InAssetPath.Len() - 1] == '/') ;
-
 	FCharacterNode* CurrentNode = &RootContentPathsTrie;
 
 	for (const TCHAR& Character : InAssetPath)
@@ -3685,6 +3679,8 @@ void UContentBrowserAssetDataSource::AddRootContentPathToStateMachine(const FStr
 		++NextNode.Value;
 		CurrentNode = NextNode.Key.Get();
 	}
+
+	CurrentNode->bIsEndOfAMountPoint = true;
 }
 
 void UContentBrowserAssetDataSource::RemoveRootContentPathFromStateMachine(const FString& InAssetPath)
@@ -3711,6 +3707,8 @@ void UContentBrowserAssetDataSource::RemoveRootContentPathFromStateMachine(const
 
 		CurrentNode = NextNode->Key.Get();
 	}
+
+	CurrentNode->bIsEndOfAMountPoint = false;
 }
 
 bool UContentBrowserAssetDataSource::PathPassesCompiledDataFilter(const FContentBrowserCompiledAssetDataFilter& InFilter, const FName InInternalPath)
diff --git a/Engine/Plugins/Editor/ContentBrowser/ContentBrowserAssetDataSource/Source/ContentBrowserAssetDataSource/Public/ContentBrowserAssetDataSource.h b/Engine/Plugins/Editor/ContentBrowser/ContentBrowserAssetDataSource/Source/ContentBrowserAssetDataSource/Public/ContentBrowserAssetDataSource.h
--- a/Engine/Plugins/Editor/ContentBrowser/ContentBrowserAssetDataSource/Source/ContentBrowserAssetDataSource/Public/ContentBrowserAssetDataSource.h
+++ b/Engine/Plugins/Editor/ContentBrowser/ContentBrowserAssetDataSource/Source/ContentBrowserAssetDataSource/Public/ContentBrowserAssetDataSource.h
@@ -516,6 +516,8 @@ private:
 	{
 		// The next characters in the tree and the number of paths beginning with the prefix including that character for use in removing paths.
 		TMap<TCHAR, TPair<FCharacterNodePtr, int32>> NextNodes;
+
+		bool bIsEndOfAMountPoint = false;
 	};
 
 	// Tree of character nodes all in lower case. Used to speed up queries against the RootContentPaths Array.
