Allow content browser data subsystem to return data as callback or array, allowing pre-sizing of output array for large queries.
ContentBrowserAssetDataSource:
- Convert asset data to content browser item data in parallel for large batches of items.
- Query on-disk and in-memory assets in parallel.
Remove mutable fields from FContentBrowserItemData for thread safety.
Add FContentBrowserMinimalItemData as a comparison type for identifying FContentBrowserItemData, e.g. when removing from FContentBrowserItem.

diff --git a/Engine/Plugins/Editor/ContentBrowser/ContentBrowserAliasDataSoure/Source/ContentBrowserAliasDataSource/Private/ContentBrowserAliasDataSource.cpp b/Engine/Plugins/Editor/ContentBrowser/ContentBrowserAliasDataSoure/Source/ContentBrowserAliasDataSource/Private/ContentBrowserAliasDataSource.cpp
--- a/Engine/Plugins/Editor/ContentBrowser/ContentBrowserAliasDataSoure/Source/ContentBrowserAliasDataSource/Private/ContentBrowserAliasDataSource.cpp
+++ b/Engine/Plugins/Editor/ContentBrowser/ContentBrowserAliasDataSoure/Source/ContentBrowserAliasDataSource/Private/ContentBrowserAliasDataSource.cpp
@@ -161,7 +161,7 @@ void UContentBrowserAliasDataSource::EnumerateItemsMatchingFilter(const FContent
 		{
 			return CreateAssetFolderItem(Path);
 		};
-		UContentBrowserAssetDataSource::EnumerateFoldersMatchingFilter(this, AssetDataFilter, InCallback, EnumerateSubPaths, CreateFolderItem);
+		UContentBrowserAssetDataSource::EnumerateFoldersMatchingFilter(this, AssetDataFilter, TGetOrEnumerateSink(InCallback), EnumerateSubPaths, CreateFolderItem);
 	}
 
 	if (EnumHasAnyFlags(InFilter.ItemTypeFilter, EContentBrowserItemTypeFilter::IncludeFiles) && !AssetDataFilter->bFilterExcludesAllAssets && !AssetDataFilter->InclusiveFilter.IsEmpty())
diff --git a/Engine/Plugins/Editor/ContentBrowser/ContentBrowserAssetDataSource/Source/ContentBrowserAssetDataSource/Private/ContentBrowserAssetDataSource.cpp b/Engine/Plugins/Editor/ContentBrowser/ContentBrowserAssetDataSource/Source/ContentBrowserAssetDataSource/Private/ContentBrowserAssetDataSource.cpp
--- a/Engine/Plugins/Editor/ContentBrowser/ContentBrowserAssetDataSource/Source/ContentBrowserAssetDataSource/Private/ContentBrowserAssetDataSource.cpp
+++ b/Engine/Plugins/Editor/ContentBrowser/ContentBrowserAssetDataSource/Source/ContentBrowserAssetDataSource/Private/ContentBrowserAssetDataSource.cpp
@@ -1,6 +1,9 @@
 // Copyright Epic Games, Inc. All Rights Reserved.
 
 #include "ContentBrowserAssetDataSource.h"
+
+#include "Algo/Transform.h"
+#include "Async/ParallelFor.h"
 #include "ContentBrowserAssetDataCore.h"
 #include "AssetRegistry/AssetRegistryModule.h"
 #include "AssetToolsModule.h"
@@ -32,6 +35,7 @@
 #include "Widgets/Input/SButton.h"
 #include "Subsystems/ImportSubsystem.h"
 #include "Widgets/Images/SImage.h"
+#include "Tasks/Task.h"
 #include "ToolMenu.h"
 #include "ToolMenuEntry.h"
 #include "ToolMenuSection.h"
@@ -42,6 +46,17 @@
 
 #define LOCTEXT_NAMESPACE "ContentBrowserAssetDataSource"
 
+namespace AssetDataSource
+{
+	bool bAllowInternalParallelism = true;
+	FAutoConsoleVariableRef CVarAllowInternalParallelism(
+		TEXT("AssetDataSource.AllowInternalParallelism"),
+		bAllowInternalParallelism,
+		TEXT("Set to 0 to disable internal parallelism inside data source in case of threading issues."),
+		ECVF_Default
+	);
+}
+
 enum class EContentBrowserFolderAttributes : uint8
 {
 	/**
@@ -1853,11 +1868,13 @@ bool PathPassesCompiledDataFilterRecursive(const FContentBrowserCompiledAssetDat
 	return true;
 }
 
-void UContentBrowserAssetDataSource::EnumerateFoldersMatchingFilter(UContentBrowserDataSource* DataSource, const FContentBrowserCompiledAssetDataFilter* AssetDataFilter, TFunctionRef<bool(FContentBrowserItemData&&)> InCallback, FSubPathEnumerationFunc SubPathEnumeration, FCreateFolderItemFunc CreateFolderItem)
+void UContentBrowserAssetDataSource::EnumerateFoldersMatchingFilter(
+	UContentBrowserDataSource* DataSource, const FContentBrowserCompiledAssetDataFilter* AssetDataFilter, 
+	const TGetOrEnumerateSink<FContentBrowserItemData>& InSink, FSubPathEnumerationFunc SubPathEnumeration, FCreateFolderItemFunc CreateFolderItem)
 {
 	if (AssetDataFilter->bRunFolderQueryOnDemand)
 	{
-		auto HandleInternalPath = [&DataSource, &InCallback, &AssetDataFilter, &SubPathEnumeration, &CreateFolderItem](
+		auto HandleInternalPath = [&DataSource, &InSink, &AssetDataFilter, &SubPathEnumeration, &CreateFolderItem](
 									  const FName InInternalPath) {
 			TArray<TPair<FName, EFolderFilterState>, TInlineAllocator<16>> PathsToScan;
 			PathsToScan.Add({ InInternalPath, EFolderFilterState::None });
@@ -1867,12 +1884,12 @@ void UContentBrowserAssetDataSource::EnumerateFoldersMatchingFilter(UContentBrow
 				EFolderFilterState ParentFilterState = PathToScan.Value;
 				SubPathEnumeration(
 					PathToScan.Key,
-					[&DataSource, &InCallback, &AssetDataFilter, &PathsToScan, &CreateFolderItem, ParentFilterState](
+					[&DataSource, &InSink, &AssetDataFilter, &PathsToScan, &CreateFolderItem, ParentFilterState](
 						FName SubPath) -> bool {
 						EFolderFilterState FilterState = ParentFilterState;
 						if (PathPassesCompiledDataFilterRecursive(*AssetDataFilter, SubPath, FilterState))
 						{
-							if (!InCallback(CreateFolderItem(SubPath)))
+							if (!InSink.ProduceItem(CreateFolderItem(SubPath)))
 							{
 								return false;
 							}
@@ -1920,13 +1937,13 @@ void UContentBrowserAssetDataSource::EnumerateFoldersMatchingFilter(UContentBrow
 			while (PathsToScan.Num() > 0)
 			{
 				const FName PathToScan = PathsToScan.Pop(EAllowShrinking::No);
-				DataSource->GetRootPathVirtualTree().EnumerateSubPaths(PathToScan, [&DataSource, &InCallback, &AssetDataFilter, &VirtualPathsPassedFilter, &PathsToScan, &HandleInternalPath, &CreateFolderItem](FName VirtualSubPath, FName InternalPath)
+				DataSource->GetRootPathVirtualTree().EnumerateSubPaths(PathToScan, [&DataSource, &InSink, &AssetDataFilter, &VirtualPathsPassedFilter, &PathsToScan, &HandleInternalPath, &CreateFolderItem](FName VirtualSubPath, FName InternalPath)
 				{
 					if (VirtualPathsPassedFilter.Contains(VirtualSubPath))
 					{
 						if (!InternalPath.IsNone())
 						{
-							if (!InCallback(CreateFolderItem(InternalPath)))
+							if (!InSink.ProduceItem(CreateFolderItem(InternalPath)))
 							{
 								return false;
 							}
@@ -1935,7 +1952,7 @@ void UContentBrowserAssetDataSource::EnumerateFoldersMatchingFilter(UContentBrow
 						}
 						else
 						{
-							if (!InCallback(DataSource->CreateVirtualFolderItem(VirtualSubPath)))
+							if (!InSink.ProduceItem(DataSource->CreateVirtualFolderItem(VirtualSubPath)))
 							{
 								return false;
 							}
@@ -1960,7 +1977,7 @@ void UContentBrowserAssetDataSource::EnumerateFoldersMatchingFilter(UContentBrow
 	{
 		for (const FName& SubPath : AssetDataFilter->CachedSubPaths)
 		{
-			if (!InCallback(CreateFolderItem(SubPath)))
+			if (!InSink.ProduceItem(CreateFolderItem(SubPath)))
 			{
 				return;
 			}
@@ -1969,6 +1986,11 @@ void UContentBrowserAssetDataSource::EnumerateFoldersMatchingFilter(UContentBrow
 }
 
 void UContentBrowserAssetDataSource::EnumerateItemsMatchingFilter(const FContentBrowserDataCompiledFilter& InFilter, TFunctionRef<bool(FContentBrowserItemData&&)> InCallback)
+{
+	EnumerateItemsMatchingFilter(InFilter, TGetOrEnumerateSink<FContentBrowserItemData>(MoveTemp(InCallback)));
+}
+
+void UContentBrowserAssetDataSource::EnumerateItemsMatchingFilter(const FContentBrowserDataCompiledFilter& InFilter, const TGetOrEnumerateSink<FContentBrowserItemData>& InSink)
 {
 	const FContentBrowserDataFilterList* FilterList = InFilter.CompiledFilters.Find(this);
 	if (!FilterList)
@@ -1992,14 +2014,14 @@ void UContentBrowserAssetDataSource::EnumerateItemsMatchingFilter(const FContent
 		{
 			return CreateAssetFolderItem(Path);
 		};
-		EnumerateFoldersMatchingFilter(this, AssetDataFilter, InCallback, EnumerateSubPaths, CreateFolderItem);
+		EnumerateFoldersMatchingFilter(this, AssetDataFilter, InSink, EnumerateSubPaths, CreateFolderItem);
 	}
 
 	if (EnumHasAnyFlags(InFilter.ItemTypeFilter, EContentBrowserItemTypeFilter::IncludeFiles) && !AssetDataFilter->bFilterExcludesAllAssets)
 	{
 		for (const FAssetData& CustomSourceAsset : AssetDataFilter->CustomSourceAssets)
 		{
-			if (!InCallback(CreateAssetFileItem(CustomSourceAsset)))
+			if (!InSink.ProduceItem(CreateAssetFileItem(CustomSourceAsset)))
 			{
 				return;
 			}
@@ -2008,7 +2030,7 @@ void UContentBrowserAssetDataSource::EnumerateItemsMatchingFilter(const FContent
 		if (const FContentBrowserCompiledUnsupportedAssetDataFilter* UnsupportedAssetDataFilter = FilterList->FindFilter<FContentBrowserCompiledUnsupportedAssetDataFilter>())
 		{
 			// Using the show unsupported asset filter
-			AssetRegistry->EnumerateAssets(UnsupportedAssetDataFilter->InclusiveFilter, [this, &InCallback, &AssetDataFilter, UnsupportedAssetDataFilter](const FAssetData& AssetData)
+			AssetRegistry->EnumerateAssets(UnsupportedAssetDataFilter->InclusiveFilter, [this, &InSink, &AssetDataFilter, UnsupportedAssetDataFilter](const FAssetData& AssetData)
 			{
 				if (ContentBrowserAssetData::IsPrimaryAsset(AssetData) && AssetData.GetOptionalOuterPathName().IsNone())
 				{
@@ -2019,34 +2041,96 @@ void UContentBrowserAssetDataSource::EnumerateItemsMatchingFilter(const FContent
 						if (!(AssetRegistry->IsAssetIncludedByFilter(AssetData, UnsupportedAssetDataFilter->ConvertIfFailInclusiveFilter) && (UnsupportedAssetDataFilter->ConvertIfFailExclusiveFilter.IsEmpty() || AssetRegistry->IsAssetExcludedByFilter(AssetData, UnsupportedAssetDataFilter->ConvertIfFailExclusiveFilter))) // Do we fail the supported filter?
 							&& (AssetRegistry->IsAssetIncludedByFilter(AssetData, UnsupportedAssetDataFilter->ShowInclusiveFilter) && (UnsupportedAssetDataFilter->ShowExclusiveFilter.IsEmpty() || AssetRegistry->IsAssetExcludedByFilter(AssetData, UnsupportedAssetDataFilter->ShowExclusiveFilter)))) // Do we pass the show filter for the unsupported asset?
 						{
-							return InCallback(CreateUnsupportedAssetFileItem(AssetData));
+							return InSink.ProduceItem(CreateUnsupportedAssetFileItem(AssetData));
 						}
 
 						// Normal item test it against the class filter
 						if ((AssetDataFilter->InclusiveFilter.ClassPaths.IsEmpty() || AssetDataFilter->InclusiveFilter.ClassPaths.Contains(AssetData.AssetClassPath)) && !AssetDataFilter->ExclusiveFilter.ClassPaths.Contains(AssetData.AssetClassPath))
 						{
-							return InCallback(CreateAssetFileItem(AssetData));
+							return InSink.ProduceItem(CreateAssetFileItem(AssetData));
 						}
 					}
 				}
 				return true;
 			});
+			return;
 		}
-		else
-		{
-			AssetRegistry->EnumerateAssets(AssetDataFilter->InclusiveFilter, [this, &InCallback, &AssetDataFilter](const FAssetData& AssetData)
+
+		auto ProduceAssets = [&AssetDataFilter, &InSink, this](TArray<FAssetData>& Assets, TSet<FName> IgnorePackageNames) {
+			InSink.ReserveMore(Assets.Num());
+			for (FAssetData& AssetData : Assets)
+			{
+				if (IgnorePackageNames.Contains(AssetData.PackageName))
 				{
-					if (ContentBrowserAssetData::IsPrimaryAsset(AssetData))
+					AssetData = FAssetData{};
+				}
+			}
+
+			if (!AssetDataFilter->ExclusiveFilter.IsEmpty())
+			{
+				for (FAssetData& AssetData : Assets)
+				{
+					if (AssetRegistry->IsAssetIncludedByFilter(AssetData, AssetDataFilter->ExclusiveFilter))
 					{
-						const bool bPassesExclusiveFilter = AssetDataFilter->ExclusiveFilter.IsEmpty() || !AssetRegistry->IsAssetIncludedByFilter(AssetData, AssetDataFilter->ExclusiveFilter);
-						if (bPassesExclusiveFilter)
-						{
-							return InCallback(CreateAssetFileItem(AssetData));
-						}
+						AssetData = FAssetData{};
+					}
+				}
+			}
+			
+			// For batches above some arbitrary threshold, run conversion in parallel
+			if (AssetDataSource::bAllowInternalParallelism && Assets.Num() > 1024 * 16)
+			{
+				TArray<FContentBrowserItemData> Converted;
+				Converted.Reserve(Assets.Num());
+				Converted.AddUninitialized(Assets.Num());
+				ParallelFor(TEXT("ConvertAssetsToContentBrowserItems"), Assets.Num(), 1024 * 16, [&Assets, &Converted, this](int32 Index) {
+					if (Assets[Index].IsValid())
+					{
+						new (&Converted[Index]) FContentBrowserItemData(CreateAssetFileItem(MoveTemp(Assets[Index])));
+					}
+					else
+					{
+						new (&Converted[Index]) FContentBrowserItemData();
 					}
-					return true;
 				});
-		}
+				for (FContentBrowserItemData& Item : Converted)
+				{
+					if (Item.IsValid())
+					{
+						InSink.ProduceItem(MoveTemp(Item));
+					}
+				}
+			}
+			else
+			{
+				for (FAssetData& AssetData : Assets)
+				{
+					if (AssetData.IsValid())
+					{
+						InSink.ProduceItem(CreateAssetFileItem(MoveTemp(AssetData)));
+					}
+				}
+			}
+		};
+
+		UE::Tasks::TTask<TArray<FAssetData>> DiskTask = UE::Tasks::Launch(UE_SOURCE_LOCATION, [AssetRegistry=AssetRegistry, &AssetDataFilter]() {
+			TArray<FAssetData> Assets;
+			FARCompiledFilter OnDiskFilter = AssetDataFilter->InclusiveFilter;
+			OnDiskFilter.bIncludeOnlyOnDiskAssets = true;
+			AssetRegistry->GetAssets(OnDiskFilter, Assets);
+			return MoveTemp(Assets);
+		});
+
+		TArray<FAssetData> InMemoryAssets;
+		AssetRegistry->GetInMemoryAssets(AssetDataFilter->InclusiveFilter, InMemoryAssets);
+
+		TSet<FName> IgnorePackages;
+		ProduceAssets(InMemoryAssets, IgnorePackages);
+		Algo::Transform(InMemoryAssets, IgnorePackages, [](const FAssetData& AssetData) { return AssetData.PackageName; });
+
+		DiskTask.BusyWait();
+		TArray<FAssetData> DiskAssets = MoveTemp(DiskTask.GetResult());
+		ProduceAssets(DiskAssets, IgnorePackages);
 	}
 }
 
@@ -2258,7 +2342,8 @@ bool UContentBrowserAssetDataSource::CreateFolder(const FName InPath, FContentBr
 		InPath,
 		*FolderItemName,
 		FText::AsCultureInvariant(FolderItemName),
-		MakeShared<FContentBrowserAssetFolderItemDataPayload>(*InternalPathString)
+		MakeShared<FContentBrowserAssetFolderItemDataPayload>(*InternalPathString),
+		FName(InternalPathString)
 		);
 
 	OutPendingItem = FContentBrowserItemDataTemporaryContext(
diff --git a/Engine/Plugins/Editor/ContentBrowser/ContentBrowserAssetDataSource/Source/ContentBrowserAssetDataSource/Public/ContentBrowserAssetDataSource.h b/Engine/Plugins/Editor/ContentBrowser/ContentBrowserAssetDataSource/Source/ContentBrowserAssetDataSource/Public/ContentBrowserAssetDataSource.h
--- a/Engine/Plugins/Editor/ContentBrowser/ContentBrowserAssetDataSource/Source/ContentBrowserAssetDataSource/Public/ContentBrowserAssetDataSource.h
+++ b/Engine/Plugins/Editor/ContentBrowser/ContentBrowserAssetDataSource/Source/ContentBrowserAssetDataSource/Public/ContentBrowserAssetDataSource.h
@@ -222,7 +222,7 @@ public:
 	 *
 	 * @return 
 	 */
-	static void EnumerateFoldersMatchingFilter(UContentBrowserDataSource* DataSource, const FContentBrowserCompiledAssetDataFilter* AssetDataFilter, TFunctionRef<bool(FContentBrowserItemData&&)> InCallback, FSubPathEnumerationFunc SubPathEnumeration, FCreateFolderItemFunc CreateFolderItem);
+	static void EnumerateFoldersMatchingFilter(UContentBrowserDataSource* DataSource, const FContentBrowserCompiledAssetDataFilter* AssetDataFilter, const TGetOrEnumerateSink<FContentBrowserItemData>& InSink, FSubPathEnumerationFunc SubPathEnumeration, FCreateFolderItemFunc CreateFolderItem);
 	
 /**
 	 * Call in DoesItemPassFilter() to check if a folder passes the compiled asset data filter.
@@ -239,6 +239,7 @@ public:
 	virtual void CompileFilter(const FName InPath, const FContentBrowserDataFilter& InFilter, FContentBrowserDataCompiledFilter& OutCompiledFilter) override;
 
 	virtual void EnumerateItemsMatchingFilter(const FContentBrowserDataCompiledFilter& InFilter, TFunctionRef<bool(FContentBrowserItemData&&)> InCallback) override;
+	virtual void EnumerateItemsMatchingFilter(const FContentBrowserDataCompiledFilter& InFilter, const TGetOrEnumerateSink<FContentBrowserItemData>& InSink) override;
 
 	virtual void EnumerateItemsAtPath(const FName InPath, const EContentBrowserItemTypeFilter InItemTypeFilter, TFunctionRef<bool(FContentBrowserItemData&&)> InCallback) override;
 
diff --git a/Engine/Source/Editor/ContentBrowser/Private/AssetViewTypes.cpp b/Engine/Source/Editor/ContentBrowser/Private/AssetViewTypes.cpp
--- a/Engine/Source/Editor/ContentBrowser/Private/AssetViewTypes.cpp
+++ b/Engine/Source/Editor/ContentBrowser/Private/AssetViewTypes.cpp
@@ -65,6 +65,15 @@ void FAssetViewItem::RemoveItemData(const FContentBrowserItemData& InItemData)
 	}
 }
 
+void FAssetViewItem::RemoveItemData(const FContentBrowserMinimalItemData& InItemKey)
+{
+	Item.TryRemove(InItemKey);
+	if (Item.IsValid())
+	{
+		ItemDataChangedEvent.Broadcast();
+	}	
+}
+
 void FAssetViewItem::ClearCachedCustomColumns()
 {
 	CachedCustomColumnData.Reset();
diff --git a/Engine/Source/Editor/ContentBrowser/Private/PathViewTypes.cpp b/Engine/Source/Editor/ContentBrowser/Private/PathViewTypes.cpp
--- a/Engine/Source/Editor/ContentBrowser/Private/PathViewTypes.cpp
+++ b/Engine/Source/Editor/ContentBrowser/Private/PathViewTypes.cpp
@@ -58,6 +58,11 @@ void FTreeItem::RemoveItemData(const FContentBrowserItemData& InItemData)
 	Item.Remove(InItemData);
 }
 
+FContentBrowserItemData FTreeItem::RemoveItemData(const FContentBrowserMinimalItemData& InItemKey)
+{
+	return Item.TryRemove(InItemKey);
+}
+
 void FTreeItem::SetItemData(FContentBrowserItem InItem)
 {
 	Item = MoveTemp(InItem);
diff --git a/Engine/Source/Editor/ContentBrowser/Private/PathViewTypes.h b/Engine/Source/Editor/ContentBrowser/Private/PathViewTypes.h
--- a/Engine/Source/Editor/ContentBrowser/Private/PathViewTypes.h
+++ b/Engine/Source/Editor/ContentBrowser/Private/PathViewTypes.h
@@ -38,6 +38,7 @@ public:
 	void AppendItemData(const FContentBrowserItemData& InItemData);
 	void RemoveItemData(const FContentBrowserItem& InItem);
 	void RemoveItemData(const FContentBrowserItemData& InItemData);
+	FContentBrowserItemData RemoveItemData(const FContentBrowserMinimalItemData& InItemKey);
 
 	/** Get the underlying Content Browser item */
 	const FContentBrowserItem& GetItem() const;
diff --git a/Engine/Source/Editor/ContentBrowser/Private/SAssetView.cpp b/Engine/Source/Editor/ContentBrowser/Private/SAssetView.cpp
--- a/Engine/Source/Editor/ContentBrowser/Private/SAssetView.cpp
+++ b/Engine/Source/Editor/ContentBrowser/Private/SAssetView.cpp
@@ -2070,7 +2070,8 @@ void SAssetView::RefreshSourceItems()
 								*FString::Printf(TEXT("/Collections/%s/%s"), ECollectionShareType::ToString(ChildCollection.Type), *ChildCollection.Name.ToString()), 
 								ChildCollection.Name, 
 								FText::FromName(ChildCollection.Name), 
-								nullptr
+								nullptr,
+								FName()
 								);
 
 							const FContentBrowserItemKey FolderItemDataKey(FolderItemData);
@@ -4858,23 +4859,23 @@ void SAssetView::HandleItemDataUpdated(TArrayView<const FContentBrowserItemDataU
 		}
 	};
 
-	auto RemoveItem = [this, &bRefreshView, &ItemsPendingInplaceFrontendFilter](const FContentBrowserItemKey& InItemDataKey, const FContentBrowserItemData& InItemData)
+	auto RemoveItem = [this, &bRefreshView, &ItemsPendingInplaceFrontendFilter](const FContentBrowserItemKey& ItemKey, const FContentBrowserMinimalItemData& ItemDataKey)
 	{
-		const uint32 ItemDataKeyHash = GetTypeHash(InItemDataKey);
+		const uint32 ItemDataKeyHash = GetTypeHash(ItemKey);
 
-		if (const TSharedPtr<FAssetViewItem>* ItemToRemovePtr = AvailableBackendItems.FindByHash(ItemDataKeyHash, InItemDataKey))
+		if (const TSharedPtr<FAssetViewItem>* ItemToRemovePtr = AvailableBackendItems.FindByHash(ItemDataKeyHash, ItemKey))
 		{
 			TSharedPtr<FAssetViewItem> ItemToRemove = *ItemToRemovePtr;
 			check(ItemToRemove);
 
 			// Only fully remove this item if every sub-item is removed (items become invalid when empty)
-			ItemToRemove->RemoveItemData(InItemData);
+			ItemToRemove->RemoveItemData(ItemDataKey);
 			if (ItemToRemove->GetItem().IsValid())
 			{
 				return;
 			}
 
-			AvailableBackendItems.RemoveByHash(ItemDataKeyHash, InItemDataKey);
+			AvailableBackendItems.RemoveByHash(ItemDataKeyHash, ItemKey);
 
 			const uint32 ItemToRemoveHash = GetTypeHash(ItemToRemove);
 
@@ -4937,19 +4938,17 @@ void SAssetView::HandleItemDataUpdated(TArrayView<const FContentBrowserItemDataU
 			}
 			else
 			{
-				RemoveItem(ItemDataKey, ItemData);
+				RemoveItem(ItemDataKey, FContentBrowserMinimalItemData(ItemData));
 			}
 			break;
 
 		case EContentBrowserItemUpdateType::Moved:
 			{
-				const FContentBrowserItemData OldMinimalItemData(ItemData.GetOwnerDataSource(), ItemData.GetItemType(), ItemDataUpdate.GetPreviousVirtualPath(), NAME_None, FText(), nullptr);
-				const FContentBrowserItemKey OldItemDataKey(OldMinimalItemData);
-				RemoveItem(OldItemDataKey, OldMinimalItemData);
-
+				const FContentBrowserMinimalItemData OldItemDataKey(ItemData.GetItemType(), ItemDataUpdate.GetPreviousVirtualPath(), ItemData.GetOwnerDataSource());
+				RemoveItem(ItemDataKey, OldItemDataKey);
 				if (bItemPassFilter)
 				{
-					AddItem(ItemDataKey, MoveTemp(ItemData));
+					AddItem(FContentBrowserItemKey(ItemData), MoveTemp(ItemData));
 				}
 				else
 				{
@@ -4959,7 +4958,7 @@ void SAssetView::HandleItemDataUpdated(TArrayView<const FContentBrowserItemDataU
 			break;
 
 		case EContentBrowserItemUpdateType::Removed:
-			RemoveItem(ItemDataKey, ItemData);
+			RemoveItem(ItemDataKey, FContentBrowserMinimalItemData(ItemData));
 			break;
 
 		default:
diff --git a/Engine/Source/Editor/ContentBrowser/Private/SPathView.cpp b/Engine/Source/Editor/ContentBrowser/Private/SPathView.cpp
--- a/Engine/Source/Editor/ContentBrowser/Private/SPathView.cpp
+++ b/Engine/Source/Editor/ContentBrowser/Private/SPathView.cpp
@@ -339,6 +339,7 @@ protected:
 	// Remove the given item data from the tree - if this results in an item having no data from any sources, the tree
 	// item is removed. If the item was removed and had a parent, the parent is returned.
 	TSharedPtr<FTreeItem> TryRemoveFolderItemInternal(const FContentBrowserItemData& InItem);
+	TSharedPtr<FTreeItem> TryRemoveFolderItemInternal(const FContentBrowserMinimalItemData& InKey);
 
 	bool PassesTextFilter(const TSharedPtr<FTreeItem>& InItem);
 
@@ -531,13 +532,8 @@ void FPathViewData::ProcessDataUpdates(TConstArrayView<FContentBrowserItemDataUp
 
 			case EContentBrowserItemUpdateType::Moved:
 			{
-				const FContentBrowserItemData OldMinimalItemData(ItemData.GetOwnerDataSource(),
-					ItemData.GetItemType(),
-					ItemDataUpdate.GetPreviousVirtualPath(),
-					NAME_None,
-					FText(),
-					nullptr);
-				TSharedPtr<FTreeItem> Parent = TryRemoveFolderItemInternal(OldMinimalItemData);
+				const FContentBrowserMinimalItemData OldItemKey(ItemData.GetItemType(), ItemDataUpdate.GetPreviousVirtualPath(), ItemData.GetOwnerDataSource());
+				TSharedPtr<FTreeItem> Parent = TryRemoveFolderItemInternal(OldItemKey);
 				if (DoesItemPassFilter(ItemData))
 				{
 					NewItems.Emplace(AddFolderItemInternal(MoveTemp(ItemData), nullptr));
@@ -753,7 +749,9 @@ TSharedRef<FTreeItem> FPathViewData::AddFolderItemInternal(FContentBrowserItemDa
 																 ParentVirtualPath,
 																 FName(FPathViews::GetPathLeaf(PathView)),
 																 FText(),
-																 nullptr),
+																 nullptr, 
+																 FName() // Assuming this is a virtual path with no internal path
+																 ),
 					OldItemsByInvariantPath);
 				VirtualPathToItem.Add(ParentVirtualPath, ParentItem);
 				InvariantPathToItem.Add(
@@ -767,19 +765,18 @@ TSharedRef<FTreeItem> FPathViewData::AddFolderItemInternal(FContentBrowserItemDa
 	return LeafItem.ToSharedRef();
 }
 
-TSharedPtr<FTreeItem> FPathViewData::TryRemoveFolderItemInternal(const FContentBrowserItemData& InItem)
+TSharedPtr<FTreeItem> FPathViewData::TryRemoveFolderItemInternal(const FContentBrowserItemData& InItemData)
 {
-	if (!InItem.IsFolder())
-	{
-		// Not a folder
-		return {};
-	}
+	return TryRemoveFolderItemInternal(FContentBrowserMinimalItemData(InItemData));
+}
 
+TSharedPtr<FTreeItem> FPathViewData::TryRemoveFolderItemInternal(const FContentBrowserMinimalItemData& InItemKey)
+{
 	// Find the folder in the tree
-	if (TSharedPtr<FTreeItem> ItemToRemove = VirtualPathToItem.FindRef(InItem.GetVirtualPath()))
+	if (TSharedPtr<FTreeItem> ItemToRemove = VirtualPathToItem.FindRef(InItemKey.GetVirtualPath()))
 	{
 		// Only fully remove this item if every sub-item is removed (items become invalid when empty)
-		ItemToRemove->RemoveItemData(InItem);
+		FContentBrowserItemData OldItemData = ItemToRemove->RemoveItemData(InItemKey);
 		if (ItemToRemove->GetItem().IsValid())
 		{
 			return {};
@@ -799,8 +796,8 @@ TSharedPtr<FTreeItem> FPathViewData::TryRemoveFolderItemInternal(const FContentB
 			VisibleRootItems.Remove(ItemToRemove);
 		}
 
-		VirtualPathToItem.Remove(InItem.GetVirtualPath());
-		InvariantPathToItem.Remove(InItem.GetInvariantPath());
+		VirtualPathToItem.Remove(InItemKey.GetVirtualPath());
+		InvariantPathToItem.Remove(OldItemData.GetInvariantPath());
 		return ItemParent;
 	}
 
@@ -1094,8 +1091,12 @@ void SPathView::Construct( const FArguments& InArgs )
 
 	if (!InArgs._DefaultPath.IsEmpty() && InternalPathPassesBlockLists(InArgs._DefaultPath))
 	{
-		const FName VirtualPath =
-			IContentBrowserDataModule::Get().GetSubsystem()->ConvertInternalPathToVirtual(*InArgs._DefaultPath);
+		const FName VirtualPath = ContentBrowserData->ConvertInternalPathToVirtual(*InArgs._DefaultPath);
+		FName InternalPath;
+		if (ContentBrowserData->TryConvertVirtualPath(VirtualPath, InternalPath) != EContentBrowserPathType::Internal)
+		{
+			InternalPath = FName();
+		}
 		if (!TreeData->FindTreeItem(VirtualPath))
 		{
 			const FString DefaultPathLeafName = FPaths::GetPathLeaf(VirtualPath.ToString());
@@ -1104,7 +1105,8 @@ void SPathView::Construct( const FArguments& InArgs )
 				VirtualPath,
 				*DefaultPathLeafName,
 				FText(),
-				nullptr));
+				nullptr,
+				InternalPath));
 		}
 
 		SetSelectedPaths({ VirtualPath.ToString() });
diff --git a/Engine/Source/Editor/ContentBrowser/Public/AssetViewTypes.h b/Engine/Source/Editor/ContentBrowser/Public/AssetViewTypes.h
--- a/Engine/Source/Editor/ContentBrowser/Public/AssetViewTypes.h
+++ b/Engine/Source/Editor/ContentBrowser/Public/AssetViewTypes.h
@@ -36,6 +36,8 @@ public:
 
 	void RemoveItemData(const FContentBrowserItemData& InItemData);
 
+	void RemoveItemData(const FContentBrowserMinimalItemData& InItemKey);
+
 	/** Clear cached custom column data */
 	void ClearCachedCustomColumns();
 
diff --git a/Engine/Source/Editor/ContentBrowserData/Private/ContentBrowserDataSource.cpp b/Engine/Source/Editor/ContentBrowserData/Private/ContentBrowserDataSource.cpp
--- a/Engine/Source/Editor/ContentBrowserData/Private/ContentBrowserDataSource.cpp
+++ b/Engine/Source/Editor/ContentBrowserData/Private/ContentBrowserDataSource.cpp
@@ -218,10 +218,21 @@ void UContentBrowserDataSource::EnumerateItemsMatchingFilter(const FContentBrows
 {
 }
 
+void UContentBrowserDataSource::EnumerateItemsMatchingFilter(const FContentBrowserDataCompiledFilter& InFilter, const TGetOrEnumerateSink<FContentBrowserItemData>& InSink)
+{
+	// Call older API function if this wasn't overridden
+	EnumerateItemsMatchingFilter(InFilter, [&InSink](FContentBrowserItemData&& Item) { return InSink.ProduceItem(MoveTemp(Item)); });
+}
+
 void UContentBrowserDataSource::EnumerateItemsAtPath(const FName InPath, const EContentBrowserItemTypeFilter InItemTypeFilter, TFunctionRef<bool(FContentBrowserItemData&&)> InCallback)
 {
 }
 
+void UContentBrowserDataSource::EnumerateItemsAtPath(const FName InPath, const EContentBrowserItemTypeFilter InItemTypeFilter, const TGetOrEnumerateSink<FContentBrowserItemData>& InSink)
+{
+	EnumerateItemsAtPath(InPath, InItemTypeFilter, [&InSink](FContentBrowserItemData&& Item) { return InSink.ProduceItem(MoveTemp(Item)); });
+}
+
 bool UContentBrowserDataSource::EnumerateItemsAtPaths(const TArrayView<FContentBrowserItemPath> InPaths, const EContentBrowserItemTypeFilter InItemTypeFilter, TFunctionRef<bool(FContentBrowserItemData&&)> InCallback)
 {
 	return true;
@@ -581,7 +592,9 @@ FContentBrowserItemData UContentBrowserDataSource::CreateVirtualFolderItem(const
 		InFolderPath,
 		*FolderItemName,
 		MoveTemp(FolderDisplayNameOverride),
-		nullptr);
+		nullptr,
+		FName() // Virtual folders have no internal path 
+		);
 }
 
 #undef LOCTEXT_NAMESPACE
diff --git a/Engine/Source/Editor/ContentBrowserData/Private/ContentBrowserDataSubsystem.cpp b/Engine/Source/Editor/ContentBrowserData/Private/ContentBrowserDataSubsystem.cpp
--- a/Engine/Source/Editor/ContentBrowserData/Private/ContentBrowserDataSubsystem.cpp
+++ b/Engine/Source/Editor/ContentBrowserData/Private/ContentBrowserDataSubsystem.cpp
@@ -9,6 +9,7 @@
 #include "Editor.h"
 #include "Features/IModularFeatures.h"
 #include "Framework/Application/SlateApplication.h"
+#include "GetOrEnumerateSink.h"
 #include "HAL/IConsoleManager.h"
 #include "IContentBrowserDataModule.h"
 #include "Interfaces/IPluginManager.h"
@@ -301,7 +302,12 @@ void UContentBrowserDataSubsystem::EnumerateItemsMatchingFilter(const FContentBr
 
 void UContentBrowserDataSubsystem::EnumerateItemsMatchingFilter(const FContentBrowserDataCompiledFilter& InFilter, TFunctionRef<bool(FContentBrowserItemData&&)> InCallback) const
 {
-	for (const auto& ActiveDataSourcePair : ActiveDataSources)
+	EnumerateItemsMatchingFilter(InFilter, TGetOrEnumerateSink<FContentBrowserItemData>(InCallback));
+}
+
+void UContentBrowserDataSubsystem::EnumerateItemsMatchingFilter(const FContentBrowserDataCompiledFilter& InFilter, const TGetOrEnumerateSink<FContentBrowserItemData>& InSink) const
+{
+	for (const TPair<FName, UContentBrowserDataSource*>& ActiveDataSourcePair : ActiveDataSources)
 	{
 		UContentBrowserDataSource* DataSource = ActiveDataSourcePair.Value;
 
@@ -315,7 +321,8 @@ void UContentBrowserDataSubsystem::EnumerateItemsMatchingFilter(const FContentBr
 					check(EnumHasAnyFlags(InFilter.ItemTypeFilter, EContentBrowserItemTypeFilter::IncludeFolders));
 
 					const FString MountLeafName = FPackageName::GetShortName(MountRootPart);
-					InCallback(FContentBrowserItemData(DataSource, EContentBrowserItemFlags::Type_Folder, MountRootPart, *MountLeafName, FText(), nullptr));
+					FName InternalPath; // Virtual folders have no internal path 
+					InSink.ProduceItem(FContentBrowserItemData(DataSource, EContentBrowserItemFlags::Type_Folder, MountRootPart, *MountLeafName, FText(), nullptr, InternalPath));
 				}
 			}
 
@@ -327,13 +334,13 @@ void UContentBrowserDataSubsystem::EnumerateItemsMatchingFilter(const FContentBr
 					for (const auto& It : VirtualFolderFilter->CachedSubPaths)
 					{
 						// how do we skip over this item if not included (Engine Content, Engine Plugins, C++ Classes, etc..)
-						InCallback(FContentBrowserItemData(It.Value));
+						InSink.ProduceItem(FContentBrowserItemData(It.Value));
 					}
 				}
 			}
 		}
 
-		DataSource->EnumerateItemsMatchingFilter(InFilter, InCallback);
+		DataSource->EnumerateItemsMatchingFilter(InFilter, InSink);
 	}
 }
 
@@ -350,11 +357,15 @@ void UContentBrowserDataSubsystem::EnumerateItemsUnderPath(const FName InPath, c
 {
 	FContentBrowserDataCompiledFilter CompiledFilter;
 	CompileFilter(InPath, InFilter, CompiledFilter);
+	EnumerateItemsMatchingFilter(CompiledFilter, MoveTemp(InCallback));
+}
 
-	EnumerateItemsMatchingFilter(CompiledFilter, [&InCallback](FContentBrowserItemData&& InItemData)
-	{
-		return InCallback(MoveTemp(InItemData));
-	});
+void UContentBrowserDataSubsystem::EnumerateItemsUnderPath(const FName InPath, const FContentBrowserDataFilter& InFilter, const TGetOrEnumerateSink<FContentBrowserItemData>& InSink) const
+{
+	FContentBrowserDataCompiledFilter CompiledFilter;
+	CompileFilter(InPath, InFilter, CompiledFilter);
+
+	EnumerateItemsMatchingFilter(CompiledFilter, InSink);
 }
 
 TArray<FContentBrowserItem> UContentBrowserDataSubsystem::GetItemsUnderPath(const FName InPath, const FContentBrowserDataFilter& InFilter) const
@@ -398,7 +409,7 @@ void UContentBrowserDataSubsystem::EnumerateItemsAtPath(const FName InPath, cons
 void UContentBrowserDataSubsystem::EnumerateItemsAtPath(const FName InPath, const EContentBrowserItemTypeFilter InItemTypeFilter, TFunctionRef<bool(FContentBrowserItemData&&)> InCallback) const
 {
 	bool bHandledVirtualFolder = false;
-	for (const auto& ActiveDataSourcePair : ActiveDataSources)
+	for (const TPair<FName, UContentBrowserDataSource*>& ActiveDataSourcePair : ActiveDataSources)
 	{
 		UContentBrowserDataSource* DataSource = ActiveDataSourcePair.Value;
 		FName InternalPath;
diff --git a/Engine/Source/Editor/ContentBrowserData/Private/ContentBrowserItem.cpp b/Engine/Source/Editor/ContentBrowserData/Private/ContentBrowserItem.cpp
--- a/Engine/Source/Editor/ContentBrowserData/Private/ContentBrowserItem.cpp
+++ b/Engine/Source/Editor/ContentBrowserData/Private/ContentBrowserItem.cpp
@@ -278,6 +278,32 @@ bool FContentBrowserItem::TryRemove(const FContentBrowserItemData& InOther, FTex
 	return true;
 }
 
+FContentBrowserItemData FContentBrowserItem::TryRemove(const FContentBrowserMinimalItemData& InKey, FText* OutError)
+{
+	const FContentBrowserItemData* PrimaryItemData = GetPrimaryInternalItem();
+	if (!PrimaryItemData)
+	{
+		// We are empty, nothing to do (but not a failure)
+		return FContentBrowserItemData();
+	}
+
+	// Remove the other item from our list
+	int32 Index = ItemDataArray.IndexOfByPredicate([Key=InKey](const FContentBrowserItemData& Data) { return Key == FContentBrowserMinimalItemData(Data); });
+	if (Index != INDEX_NONE)
+	{
+		FContentBrowserItemData Data = MoveTemp(ItemDataArray[Index]);
+		ItemDataArray.RemoveAt(Index);
+		return MoveTemp(Data);
+	}
+	else
+	{
+		if (OutError)
+		{
+			*OutError = LOCTEXT("RemoveError_KeyMismatch", "No data matching given key");
+		}
+		return FContentBrowserItemData();
+	}
+}
 bool FContentBrowserItem::IsFolder() const
 {
 	const FContentBrowserItemData* PrimaryItemData = GetPrimaryInternalItem();
diff --git a/Engine/Source/Editor/ContentBrowserData/Private/ContentBrowserItemData.cpp b/Engine/Source/Editor/ContentBrowserData/Private/ContentBrowserItemData.cpp
--- a/Engine/Source/Editor/ContentBrowserData/Private/ContentBrowserItemData.cpp
+++ b/Engine/Source/Editor/ContentBrowserData/Private/ContentBrowserItemData.cpp
@@ -16,7 +16,7 @@ FContentBrowserItemData::FContentBrowserItemData(
 	FName InItemName,
 	FText InDisplayNameOverride,
 	TSharedPtr<const IContentBrowserItemDataPayload> InPayload,
-	TOptional<FName> InInternalPath)
+	FName InInternalPath)
 	: OwnerDataSource(InOwnerDataSource)
 	, ItemFlags(InItemFlags)
 	, VirtualPath(InVirtualPath)
@@ -26,6 +26,10 @@ FContentBrowserItemData::FContentBrowserItemData(
 	, Payload(MoveTemp(InPayload))
 {
 	checkf(IsFolder() != IsFile(), TEXT("Items must be either a folder or a file!"));
+	if (CachedDisplayName.IsEmpty())
+	{
+		CachedDisplayName = FText::AsCultureInvariant(ItemName.ToString());
+	}
 }
 
 bool FContentBrowserItemData::operator==(const FContentBrowserItemData& InOther) const
@@ -118,24 +122,7 @@ FName FContentBrowserItemData::GetInvariantPath() const
 
 FName FContentBrowserItemData::GetInternalPath() const
 {
-	if (InternalPath.IsSet())
-	{
-		return InternalPath.GetValue();
-	}
-	if (UContentBrowserDataSource* DataSource = OwnerDataSource.Get())
-	{
-		if (!VirtualPath.IsNone())
-		{
-			FName ConvertedPath;
-			if (DataSource->TryConvertVirtualPath(VirtualPath, ConvertedPath) == EContentBrowserPathType::Internal)
-			{
-				InternalPath = ConvertedPath;
-				return ConvertedPath;
-			}
-		}
-	}
-
-	return NAME_None;
+	return InternalPath;
 }
 
 FName FContentBrowserItemData::GetItemName() const
@@ -145,10 +132,6 @@ FName FContentBrowserItemData::GetItemName() const
 
 FText FContentBrowserItemData::GetDisplayName() const
 {
-	if (CachedDisplayName.IsEmpty())
-	{
-		CachedDisplayName = FText::AsCultureInvariant(ItemName.ToString());
-	}
 	return CachedDisplayName;
 }
 
@@ -341,6 +324,28 @@ FContentBrowserItemData FContentBrowserItemDataTemporaryContext::FinalizeItem(co
 	return OnFinalizeItem.Execute(ItemData, InProposedName, OutErrorMsg);
 }
 
+FContentBrowserMinimalItemData::FContentBrowserMinimalItemData(const FContentBrowserItemData& InItemData)
+	: FContentBrowserMinimalItemData(InItemData.GetItemType(), InItemData.GetVirtualPath(), InItemData.GetOwnerDataSource())
+{
+
+}
+
+FContentBrowserMinimalItemData::FContentBrowserMinimalItemData(EContentBrowserItemFlags InItemType, FName InVirtualPath, const UContentBrowserDataSource* InSource)
+	: ItemType(InItemType)
+	, VirtualPath(InVirtualPath)
+	, DataSource(InSource)
+{
+
+}
+
+FString FContentBrowserMinimalItemData::ToString() const
+{
+	return FString::Printf(TEXT("%s:%s:%s"), 
+		*VirtualPath.ToString(), 
+		EnumHasAllFlags(ItemType, EContentBrowserItemFlags::Type_File) ? TEXT("File") : TEXT("Folder"),
+		DataSource ? *DataSource->GetName() : TEXT("null")
+	);
+}
 
 FContentBrowserItemDataKey::FContentBrowserItemDataKey(const FContentBrowserItemData& InItemData)
 	: ItemType(InItemData.GetItemType())
diff --git a/Engine/Source/Editor/ContentBrowserData/Public/ContentBrowserDataSource.h b/Engine/Source/Editor/ContentBrowserData/Public/ContentBrowserDataSource.h
--- a/Engine/Source/Editor/ContentBrowserData/Public/ContentBrowserDataSource.h
+++ b/Engine/Source/Editor/ContentBrowserData/Public/ContentBrowserDataSource.h
@@ -198,8 +198,10 @@ public:
 	 *
 	 * @param InFilter The compiled filter used to find matching items.
 	 * @param InCallback The function to invoke for each matching item (return true to continue enumeration).
+	 * @param InSink Parameter provided to receive output as either a callback or an array.
 	 */
 	virtual void EnumerateItemsMatchingFilter(const FContentBrowserDataCompiledFilter& InFilter, TFunctionRef<bool(FContentBrowserItemData&&)> InCallback);
+	virtual void EnumerateItemsMatchingFilter(const FContentBrowserDataCompiledFilter& InFilter, const TGetOrEnumerateSink<FContentBrowserItemData>& InSink);
 
 	/**
 	 * Enumerate items that have the given virtual path, optionally filtering by type, and invoking the callback for each matching item.
@@ -208,9 +210,10 @@ public:
 	 * @param InPath The virtual path to find items for.
 	 * @param InItemTypeFilter The types of items we want to find.
 	 * @param InCallback The function to invoke for each matching item (return true to continue enumeration).
+	 * @param InSink Parameter provided to receive output as either a callback or an array.
 	 */
 	virtual void EnumerateItemsAtPath(const FName InPath, const EContentBrowserItemTypeFilter InItemTypeFilter, TFunctionRef<bool(FContentBrowserItemData&&)> InCallback);
-
+	virtual void EnumerateItemsAtPath(const FName InPath, const EContentBrowserItemTypeFilter InItemTypeFilter, const TGetOrEnumerateSink<FContentBrowserItemData>& InSink);
 	
 	/**
 	 * Enumerate the items (folders and/or files) that exist at the given content browser paths.
diff --git a/Engine/Source/Editor/ContentBrowserData/Public/ContentBrowserDataSubsystem.h b/Engine/Source/Editor/ContentBrowserData/Public/ContentBrowserDataSubsystem.h
--- a/Engine/Source/Editor/ContentBrowserData/Public/ContentBrowserDataSubsystem.h
+++ b/Engine/Source/Editor/ContentBrowserData/Public/ContentBrowserDataSubsystem.h
@@ -18,6 +18,7 @@
 #include "CoreMinimal.h"
 #include "Delegates/Delegate.h"
 #include "EditorSubsystem.h"
+#include "GetOrEnumerateSink.h"
 #include "HAL/Platform.h"
 #include "Misc/AssertionMacros.h"
 #include "Misc/NamePermissionList.h"
@@ -177,14 +178,26 @@ public:
 	/**
 	 * Enumerate the items (folders and/or files) that match a previously compiled filter.
 	 */
-	void EnumerateItemsMatchingFilter(const FContentBrowserDataCompiledFilter& InFilter, TFunctionRef<bool(FContentBrowserItem&&)> InCallback) const;
-	void EnumerateItemsMatchingFilter(const FContentBrowserDataCompiledFilter& InFilter, TFunctionRef<bool(FContentBrowserItemData&&)> InCallback) const;
+	void EnumerateItemsMatchingFilter(
+		const FContentBrowserDataCompiledFilter& InFilter, TFunctionRef<bool(FContentBrowserItem&&)> InCallback) const;
+	void EnumerateItemsMatchingFilter(const FContentBrowserDataCompiledFilter& InFilter,
+		TFunctionRef<bool(FContentBrowserItemData&&)> InCallback) const;
+	void EnumerateItemsMatchingFilter(
+		const FContentBrowserDataCompiledFilter& InFilter, const TGetOrEnumerateSink<FContentBrowserItemData>& InSink) const;
 
 	/**
 	 * Enumerate the items (folders and/or files) that exist under the given virtual path.
 	 */
-	void EnumerateItemsUnderPath(const FName InPath, const FContentBrowserDataFilter& InFilter, TFunctionRef<bool(FContentBrowserItem&&)> InCallback) const;
-	void EnumerateItemsUnderPath(const FName InPath, const FContentBrowserDataFilter& InFilter, TFunctionRef<bool(FContentBrowserItemData&&)> InCallback) const;
+	void EnumerateItemsUnderPath(const FName InPath,
+		const FContentBrowserDataFilter& InFilter,
+		TFunctionRef<bool(FContentBrowserItem&&)> InCallback) const;
+	void EnumerateItemsUnderPath(const FName InPath,
+		const FContentBrowserDataFilter& InFilter,
+		TFunctionRef<bool(FContentBrowserItemData&&)> InCallback) const;
+	// Version which allows passing an array or a callback, allows sources to presize the array for large queries
+	void EnumerateItemsUnderPath(const FName InPath,
+		const FContentBrowserDataFilter& InFilter,
+		const TGetOrEnumerateSink<FContentBrowserItemData>& InSink) const;
 
 	/**
 	 * Get the items (folders and/or files) that exist under the given virtual path.
diff --git a/Engine/Source/Editor/ContentBrowserData/Public/ContentBrowserItem.h b/Engine/Source/Editor/ContentBrowserData/Public/ContentBrowserItem.h
--- a/Engine/Source/Editor/ContentBrowserData/Public/ContentBrowserItem.h
+++ b/Engine/Source/Editor/ContentBrowserData/Public/ContentBrowserItem.h
@@ -20,6 +20,7 @@
 #include "ContentBrowserItem.generated.h"
 
 class FAssetThumbnail;
+class FContentBrowserMinimalItemData;
 class UContentBrowserDataSource;
 class UObject;
 struct FAssetData;
@@ -113,6 +114,12 @@ public:
 	 */
 	bool TryRemove(const FContentBrowserItemData& InOther, FText* OutError = nullptr);
 
+	/**
+	 * Attempt to remove an item matching the given minimal representation from this one, providing error information if the removal wasn't possible.
+	 * @return A valid item data (IsValid) if anything was removed, an invalid one otherwise. 
+	 */
+	FContentBrowserItemData TryRemove(const FContentBrowserMinimalItemData& InKey, FText* OutError = nullptr);
+
 	/**
 	 * Get the array of internal items that comprise this composite item.
 	 */
diff --git a/Engine/Source/Editor/ContentBrowserData/Public/ContentBrowserItemData.h b/Engine/Source/Editor/ContentBrowserData/Public/ContentBrowserItemData.h
--- a/Engine/Source/Editor/ContentBrowserData/Public/ContentBrowserItemData.h
+++ b/Engine/Source/Editor/ContentBrowserData/Public/ContentBrowserItemData.h
@@ -19,6 +19,7 @@
 #include "UObject/WeakObjectPtr.h"
 #include "UObject/WeakObjectPtrTemplates.h"
 
+class FContentBrowserMinimalItemData;
 class IContentBrowserItemDataSink;
 class UContentBrowserDataSource;
 
@@ -134,7 +135,7 @@ public:
 	 * @param InDisplayNameOverride The user-facing name of this item (eg, "MyFile"). This will be lazily set to
 	 * InItemName if no override is provided.
 	 * @param InPayload Any data source defined payload data for this item.
-	 * @param InInternalPath Precalculated internal path if available
+	 * @param InInternalPath Internal path if any (some items such as "virtual folders" may not have an internal path)
 	 */
 	FContentBrowserItemData(
 		UContentBrowserDataSource* InOwnerDataSource,
@@ -143,7 +144,7 @@ public:
 		FName InItemName,
 		FText InDisplayNameOverride,
 		TSharedPtr<const IContentBrowserItemDataPayload> InPayload,
-		TOptional<FName> InInternalPath = {});
+		FName InInternalPath);
 
 	/**
 	 * Copy support.
@@ -278,11 +279,11 @@ private:
 	/** The leaf-name of this item (eg, "MyFile") */
 	FName ItemName;
 
-	/** Cached/pre-computed value for GetInternalPath */
-	mutable TOptional<FName> InternalPath;
+	/** Internal path used to identify item within the owner data source */
+	FName InternalPath;
 
 	/** The user-facing name of this item (eg, "MyFile") */
-	mutable FText CachedDisplayName;
+	FText CachedDisplayName;
 
 	/** Any data source defined payload data for this item */
 	TSharedPtr<const IContentBrowserItemDataPayload> Payload;
@@ -624,6 +625,60 @@ private:
 	FOnFinalizeItem OnFinalizeItem;
 };
 
+/**
+ * Minimal representation of a FContentBrowserItemData instance for comparison purposes.
+ * @note Unlike FContentBrowserItemDataKey this does consider the data source, and unlike FContentBrowserItemKey it considers the data source for folder items.
+ *  It can be used to identify a FContentBrowserItemData for removal from an FContentBrowserItem
+ */
+class CONTENTBROWSERDATA_API FContentBrowserMinimalItemData
+{
+public:
+	explicit FContentBrowserMinimalItemData(const FContentBrowserItemData& InItemData);
+	FContentBrowserMinimalItemData(EContentBrowserItemFlags InItemType, FName InVirtualPath, const UContentBrowserDataSource* InSource);
+
+	bool operator==(const FContentBrowserMinimalItemData& Other) const
+	{
+		return !(*this != Other);
+	}
+	bool operator!=(const FContentBrowserMinimalItemData& Other) const
+	{
+		return ItemType == Other.ItemType && VirtualPath == Other.VirtualPath && DataSource == Other.DataSource;
+	}
+
+	/** Return a string representation for debugging. */
+	FString ToString() const;
+	
+	EContentBrowserItemFlags GetItemType() const
+	{
+		return ItemType;
+	} 
+
+	FName GetVirtualPath() const
+	{
+		return VirtualPath;
+	}
+
+	const UContentBrowserDataSource* GetDataSource() const 
+	{
+		return DataSource;
+	}
+
+private:
+	/**
+	 * Flags denoting the item type information for an item instance.
+	 * @note This is always masked against EContentBrowserItemFlags::Type_MASK to remove any non-type information.
+	 */
+	EContentBrowserItemFlags ItemType = EContentBrowserItemFlags::None;
+
+	/**
+	 * The complete virtual path that uniquely identifies an item within its owner data source.
+	 */
+	FName VirtualPath;
+
+	/** A pointer to the data source that manages the thing represented by this key */
+	const UContentBrowserDataSource* DataSource = nullptr;
+};
+
 /**
  * Minimal representation of a FContentBrowserItemData instance that can be used as a map key.
  * @note This key doesn't consider the data source, so should only be used for maps within a given data source.
diff --git a/Engine/Source/Editor/ContentBrowserData/Public/GetOrEnumerateSink.h b/Engine/Source/Editor/ContentBrowserData/Public/GetOrEnumerateSink.h
new file mode 100644
--- /dev/null
+++ b/Engine/Source/Editor/ContentBrowserData/Public/GetOrEnumerateSink.h
@@ -0,0 +1,60 @@
+#pragma once
+
+#include "Containers/Array.h"
+#include "Templates/Function.h"
+#include "Templates/Invoke.h"
+
+/**
+ * Utility class to simplify implementation of data sources that can be called as 'get'
+ * (fill an array) or 'enumerate' (call a callback).
+ * 
+ * Usage: Functions should take TGetOrEnumerateSink by value or const ref to allow in-place construction.
+ * 
+ * @tparam ItemType The type of item produced. The sink can be constructed with TArray<ItemType> or TFunctionRef<ItemType&&>.
+ */
+template<typename ItemType>
+class TGetOrEnumerateSink
+{
+public:
+    explicit TGetOrEnumerateSink(TFunctionRef<bool(ItemType&&)> InCallback)
+        : Callback(MoveTemp(InCallback))
+    {
+    }
+    explicit TGetOrEnumerateSink(TArray<ItemType>& InOutput)
+        : Array(&InOutput)
+    {
+    }
+
+    // Allow copying for delegation to other methods if desired. Just copies pointers into sink. No state needs to be maintained between copies. 
+    TGetOrEnumerateSink(const TGetOrEnumerateSink&) = default;
+    TGetOrEnumerateSink& operator=(const TGetOrEnumerateSink&) = default;
+    
+    // Allow moving to delegate the sink into other calls 
+    TGetOrEnumerateSink(TGetOrEnumerateSink&&) = default;
+    TGetOrEnumerateSink& operator=(TGetOrEnumerateSink&&) = default;
+
+    void ReserveMore(int32 NumAdditionalItems) const
+    {
+        if (Array)
+        {
+            Array->Reserve(Array->Num() + NumAdditionalItems);
+        }
+    }
+
+    bool ProduceItem(ItemType&& Item) const
+    {
+        if (Array)
+        {
+            Array->Add(MoveTemp(Item));
+            return true;
+        }
+        else
+        {
+            return ::Invoke(Callback.GetValue(), MoveTemp(Item));
+        }
+    }
+
+private:
+    TArray<ItemType>* Array = nullptr;
+    TOptional<TFunctionRef<bool(ItemType&&)>> Callback;
+};
\ No newline at end of file
