[Backout] - CL33234860
Original CL Desc
-----------------------------------------------------------------
Allow content browser data subsystem to return data as callback or array, allowing pre-sizing of output array for large queries.
ContentBrowserAssetDataSource:
- Convert asset data to content browser item data in parallel for large batches of items.
- Query on-disk and in-memory assets in parallel.
Remove mutable fields from FContentBrowserItemData for thread safety.
Add FContentBrowserMinimalItemData as a comparison type for identifying FContentBrowserItemData, e.g. when removing from FContentBrowserItem.

diff --git a/Engine/Plugins/Editor/ContentBrowser/ContentBrowserAliasDataSoure/Source/ContentBrowserAliasDataSource/Private/ContentBrowserAliasDataSource.cpp b/Engine/Plugins/Editor/ContentBrowser/ContentBrowserAliasDataSoure/Source/ContentBrowserAliasDataSource/Private/ContentBrowserAliasDataSource.cpp
--- a/Engine/Plugins/Editor/ContentBrowser/ContentBrowserAliasDataSoure/Source/ContentBrowserAliasDataSource/Private/ContentBrowserAliasDataSource.cpp
+++ b/Engine/Plugins/Editor/ContentBrowser/ContentBrowserAliasDataSoure/Source/ContentBrowserAliasDataSource/Private/ContentBrowserAliasDataSource.cpp
@@ -161,7 +161,7 @@ void UContentBrowserAliasDataSource::EnumerateItemsMatchingFilter(const FContent
 		{
 			return CreateAssetFolderItem(Path);
 		};
-		UContentBrowserAssetDataSource::EnumerateFoldersMatchingFilter(this, AssetDataFilter, TGetOrEnumerateSink(InCallback), EnumerateSubPaths, CreateFolderItem);
+		UContentBrowserAssetDataSource::EnumerateFoldersMatchingFilter(this, AssetDataFilter, InCallback, EnumerateSubPaths, CreateFolderItem);
 	}
 
 	if (EnumHasAnyFlags(InFilter.ItemTypeFilter, EContentBrowserItemTypeFilter::IncludeFiles) && !AssetDataFilter->bFilterExcludesAllAssets && !AssetDataFilter->InclusiveFilter.IsEmpty())
diff --git a/Engine/Plugins/Editor/ContentBrowser/ContentBrowserAssetDataSource/Source/ContentBrowserAssetDataSource/Private/ContentBrowserAssetDataSource.cpp b/Engine/Plugins/Editor/ContentBrowser/ContentBrowserAssetDataSource/Source/ContentBrowserAssetDataSource/Private/ContentBrowserAssetDataSource.cpp
--- a/Engine/Plugins/Editor/ContentBrowser/ContentBrowserAssetDataSource/Source/ContentBrowserAssetDataSource/Private/ContentBrowserAssetDataSource.cpp
+++ b/Engine/Plugins/Editor/ContentBrowser/ContentBrowserAssetDataSource/Source/ContentBrowserAssetDataSource/Private/ContentBrowserAssetDataSource.cpp
@@ -1,9 +1,6 @@
 // Copyright Epic Games, Inc. All Rights Reserved.
 
 #include "ContentBrowserAssetDataSource.h"
-
-#include "Algo/Transform.h"
-#include "Async/ParallelFor.h"
 #include "ContentBrowserAssetDataCore.h"
 #include "AssetRegistry/AssetRegistryModule.h"
 #include "AssetToolsModule.h"
@@ -35,7 +32,6 @@
 #include "Widgets/Input/SButton.h"
 #include "Subsystems/ImportSubsystem.h"
 #include "Widgets/Images/SImage.h"
-#include "Tasks/Task.h"
 #include "ToolMenu.h"
 #include "ToolMenuEntry.h"
 #include "ToolMenuSection.h"
@@ -46,17 +42,6 @@
 
 #define LOCTEXT_NAMESPACE "ContentBrowserAssetDataSource"
 
-namespace AssetDataSource
-{
-	bool bAllowInternalParallelism = true;
-	FAutoConsoleVariableRef CVarAllowInternalParallelism(
-		TEXT("AssetDataSource.AllowInternalParallelism"),
-		bAllowInternalParallelism,
-		TEXT("Set to 0 to disable internal parallelism inside data source in case of threading issues."),
-		ECVF_Default
-	);
-}
-
 enum class EContentBrowserFolderAttributes : uint8
 {
 	/**
@@ -1868,13 +1853,11 @@ bool PathPassesCompiledDataFilterRecursive(const FContentBrowserCompiledAssetDat
 	return true;
 }
 
-void UContentBrowserAssetDataSource::EnumerateFoldersMatchingFilter(
-	UContentBrowserDataSource* DataSource, const FContentBrowserCompiledAssetDataFilter* AssetDataFilter, 
-	const TGetOrEnumerateSink<FContentBrowserItemData>& InSink, FSubPathEnumerationFunc SubPathEnumeration, FCreateFolderItemFunc CreateFolderItem)
+void UContentBrowserAssetDataSource::EnumerateFoldersMatchingFilter(UContentBrowserDataSource* DataSource, const FContentBrowserCompiledAssetDataFilter* AssetDataFilter, TFunctionRef<bool(FContentBrowserItemData&&)> InCallback, FSubPathEnumerationFunc SubPathEnumeration, FCreateFolderItemFunc CreateFolderItem)
 {
 	if (AssetDataFilter->bRunFolderQueryOnDemand)
 	{
-		auto HandleInternalPath = [&DataSource, &InSink, &AssetDataFilter, &SubPathEnumeration, &CreateFolderItem](
+		auto HandleInternalPath = [&DataSource, &InCallback, &AssetDataFilter, &SubPathEnumeration, &CreateFolderItem](
 									  const FName InInternalPath) {
 			TArray<TPair<FName, EFolderFilterState>, TInlineAllocator<16>> PathsToScan;
 			PathsToScan.Add({ InInternalPath, EFolderFilterState::None });
@@ -1884,12 +1867,12 @@ void UContentBrowserAssetDataSource::EnumerateFoldersMatchingFilter(
 				EFolderFilterState ParentFilterState = PathToScan.Value;
 				SubPathEnumeration(
 					PathToScan.Key,
-					[&DataSource, &InSink, &AssetDataFilter, &PathsToScan, &CreateFolderItem, ParentFilterState](
+					[&DataSource, &InCallback, &AssetDataFilter, &PathsToScan, &CreateFolderItem, ParentFilterState](
 						FName SubPath) -> bool {
 						EFolderFilterState FilterState = ParentFilterState;
 						if (PathPassesCompiledDataFilterRecursive(*AssetDataFilter, SubPath, FilterState))
 						{
-							if (!InSink.ProduceItem(CreateFolderItem(SubPath)))
+							if (!InCallback(CreateFolderItem(SubPath)))
 							{
 								return false;
 							}
@@ -1937,13 +1920,13 @@ void UContentBrowserAssetDataSource::EnumerateFoldersMatchingFilter(
 			while (PathsToScan.Num() > 0)
 			{
 				const FName PathToScan = PathsToScan.Pop(EAllowShrinking::No);
-				DataSource->GetRootPathVirtualTree().EnumerateSubPaths(PathToScan, [&DataSource, &InSink, &AssetDataFilter, &VirtualPathsPassedFilter, &PathsToScan, &HandleInternalPath, &CreateFolderItem](FName VirtualSubPath, FName InternalPath)
+				DataSource->GetRootPathVirtualTree().EnumerateSubPaths(PathToScan, [&DataSource, &InCallback, &AssetDataFilter, &VirtualPathsPassedFilter, &PathsToScan, &HandleInternalPath, &CreateFolderItem](FName VirtualSubPath, FName InternalPath)
 				{
 					if (VirtualPathsPassedFilter.Contains(VirtualSubPath))
 					{
 						if (!InternalPath.IsNone())
 						{
-							if (!InSink.ProduceItem(CreateFolderItem(InternalPath)))
+							if (!InCallback(CreateFolderItem(InternalPath)))
 							{
 								return false;
 							}
@@ -1952,7 +1935,7 @@ void UContentBrowserAssetDataSource::EnumerateFoldersMatchingFilter(
 						}
 						else
 						{
-							if (!InSink.ProduceItem(DataSource->CreateVirtualFolderItem(VirtualSubPath)))
+							if (!InCallback(DataSource->CreateVirtualFolderItem(VirtualSubPath)))
 							{
 								return false;
 							}
@@ -1977,7 +1960,7 @@ void UContentBrowserAssetDataSource::EnumerateFoldersMatchingFilter(
 	{
 		for (const FName& SubPath : AssetDataFilter->CachedSubPaths)
 		{
-			if (!InSink.ProduceItem(CreateFolderItem(SubPath)))
+			if (!InCallback(CreateFolderItem(SubPath)))
 			{
 				return;
 			}
@@ -1986,11 +1969,6 @@ void UContentBrowserAssetDataSource::EnumerateFoldersMatchingFilter(
 }
 
 void UContentBrowserAssetDataSource::EnumerateItemsMatchingFilter(const FContentBrowserDataCompiledFilter& InFilter, TFunctionRef<bool(FContentBrowserItemData&&)> InCallback)
-{
-	EnumerateItemsMatchingFilter(InFilter, TGetOrEnumerateSink<FContentBrowserItemData>(MoveTemp(InCallback)));
-}
-
-void UContentBrowserAssetDataSource::EnumerateItemsMatchingFilter(const FContentBrowserDataCompiledFilter& InFilter, const TGetOrEnumerateSink<FContentBrowserItemData>& InSink)
 {
 	const FContentBrowserDataFilterList* FilterList = InFilter.CompiledFilters.Find(this);
 	if (!FilterList)
@@ -2014,14 +1992,14 @@ void UContentBrowserAssetDataSource::EnumerateItemsMatchingFilter(const FContent
 		{
 			return CreateAssetFolderItem(Path);
 		};
-		EnumerateFoldersMatchingFilter(this, AssetDataFilter, InSink, EnumerateSubPaths, CreateFolderItem);
+		EnumerateFoldersMatchingFilter(this, AssetDataFilter, InCallback, EnumerateSubPaths, CreateFolderItem);
 	}
 
 	if (EnumHasAnyFlags(InFilter.ItemTypeFilter, EContentBrowserItemTypeFilter::IncludeFiles) && !AssetDataFilter->bFilterExcludesAllAssets)
 	{
 		for (const FAssetData& CustomSourceAsset : AssetDataFilter->CustomSourceAssets)
 		{
-			if (!InSink.ProduceItem(CreateAssetFileItem(CustomSourceAsset)))
+			if (!InCallback(CreateAssetFileItem(CustomSourceAsset)))
 			{
 				return;
 			}
@@ -2030,7 +2008,7 @@ void UContentBrowserAssetDataSource::EnumerateItemsMatchingFilter(const FContent
 		if (const FContentBrowserCompiledUnsupportedAssetDataFilter* UnsupportedAssetDataFilter = FilterList->FindFilter<FContentBrowserCompiledUnsupportedAssetDataFilter>())
 		{
 			// Using the show unsupported asset filter
-			AssetRegistry->EnumerateAssets(UnsupportedAssetDataFilter->InclusiveFilter, [this, &InSink, &AssetDataFilter, UnsupportedAssetDataFilter](const FAssetData& AssetData)
+			AssetRegistry->EnumerateAssets(UnsupportedAssetDataFilter->InclusiveFilter, [this, &InCallback, &AssetDataFilter, UnsupportedAssetDataFilter](const FAssetData& AssetData)
 			{
 				if (ContentBrowserAssetData::IsPrimaryAsset(AssetData) && AssetData.GetOptionalOuterPathName().IsNone())
 				{
@@ -2041,96 +2019,34 @@ void UContentBrowserAssetDataSource::EnumerateItemsMatchingFilter(const FContent
 						if (!(AssetRegistry->IsAssetIncludedByFilter(AssetData, UnsupportedAssetDataFilter->ConvertIfFailInclusiveFilter) && (UnsupportedAssetDataFilter->ConvertIfFailExclusiveFilter.IsEmpty() || AssetRegistry->IsAssetExcludedByFilter(AssetData, UnsupportedAssetDataFilter->ConvertIfFailExclusiveFilter))) // Do we fail the supported filter?
 							&& (AssetRegistry->IsAssetIncludedByFilter(AssetData, UnsupportedAssetDataFilter->ShowInclusiveFilter) && (UnsupportedAssetDataFilter->ShowExclusiveFilter.IsEmpty() || AssetRegistry->IsAssetExcludedByFilter(AssetData, UnsupportedAssetDataFilter->ShowExclusiveFilter)))) // Do we pass the show filter for the unsupported asset?
 						{
-							return InSink.ProduceItem(CreateUnsupportedAssetFileItem(AssetData));
+							return InCallback(CreateUnsupportedAssetFileItem(AssetData));
 						}
 
 						// Normal item test it against the class filter
 						if ((AssetDataFilter->InclusiveFilter.ClassPaths.IsEmpty() || AssetDataFilter->InclusiveFilter.ClassPaths.Contains(AssetData.AssetClassPath)) && !AssetDataFilter->ExclusiveFilter.ClassPaths.Contains(AssetData.AssetClassPath))
 						{
-							return InSink.ProduceItem(CreateAssetFileItem(AssetData));
+							return InCallback(CreateAssetFileItem(AssetData));
 						}
 					}
 				}
 				return true;
 			});
-			return;
 		}
-
-		auto ProduceAssets = [&AssetDataFilter, &InSink, this](TArray<FAssetData>& Assets, TSet<FName> IgnorePackageNames) {
-			InSink.ReserveMore(Assets.Num());
-			for (FAssetData& AssetData : Assets)
-			{
-				if (IgnorePackageNames.Contains(AssetData.PackageName))
-				{
-					AssetData = FAssetData{};
-				}
-			}
-
-			if (!AssetDataFilter->ExclusiveFilter.IsEmpty())
-			{
-				for (FAssetData& AssetData : Assets)
+		else
+		{
+			AssetRegistry->EnumerateAssets(AssetDataFilter->InclusiveFilter, [this, &InCallback, &AssetDataFilter](const FAssetData& AssetData)
 				{
-					if (AssetRegistry->IsAssetIncludedByFilter(AssetData, AssetDataFilter->ExclusiveFilter))
-					{
-						AssetData = FAssetData{};
-					}
-				}
-			}
-			
-			// For batches above some arbitrary threshold, run conversion in parallel
-			if (AssetDataSource::bAllowInternalParallelism && Assets.Num() > 1024 * 16)
-			{
-				TArray<FContentBrowserItemData> Converted;
-				Converted.Reserve(Assets.Num());
-				Converted.AddUninitialized(Assets.Num());
-				ParallelFor(TEXT("ConvertAssetsToContentBrowserItems"), Assets.Num(), 1024 * 16, [&Assets, &Converted, this](int32 Index) {
-					if (Assets[Index].IsValid())
+					if (ContentBrowserAssetData::IsPrimaryAsset(AssetData))
 					{
-						new (&Converted[Index]) FContentBrowserItemData(CreateAssetFileItem(MoveTemp(Assets[Index])));
-					}
-					else
-					{
-						new (&Converted[Index]) FContentBrowserItemData();
+						const bool bPassesExclusiveFilter = AssetDataFilter->ExclusiveFilter.IsEmpty() || !AssetRegistry->IsAssetIncludedByFilter(AssetData, AssetDataFilter->ExclusiveFilter);
+						if (bPassesExclusiveFilter)
+						{
+							return InCallback(CreateAssetFileItem(AssetData));
+						}
 					}
+					return true;
 				});
-				for (FContentBrowserItemData& Item : Converted)
-				{
-					if (Item.IsValid())
-					{
-						InSink.ProduceItem(MoveTemp(Item));
-					}
-				}
-			}
-			else
-			{
-				for (FAssetData& AssetData : Assets)
-				{
-					if (AssetData.IsValid())
-					{
-						InSink.ProduceItem(CreateAssetFileItem(MoveTemp(AssetData)));
-					}
-				}
-			}
-		};
-
-		UE::Tasks::TTask<TArray<FAssetData>> DiskTask = UE::Tasks::Launch(UE_SOURCE_LOCATION, [AssetRegistry=AssetRegistry, &AssetDataFilter]() {
-			TArray<FAssetData> Assets;
-			FARCompiledFilter OnDiskFilter = AssetDataFilter->InclusiveFilter;
-			OnDiskFilter.bIncludeOnlyOnDiskAssets = true;
-			AssetRegistry->GetAssets(OnDiskFilter, Assets);
-			return MoveTemp(Assets);
-		});
-
-		TArray<FAssetData> InMemoryAssets;
-		AssetRegistry->GetInMemoryAssets(AssetDataFilter->InclusiveFilter, InMemoryAssets);
-
-		TSet<FName> IgnorePackages;
-		ProduceAssets(InMemoryAssets, IgnorePackages);
-		Algo::Transform(InMemoryAssets, IgnorePackages, [](const FAssetData& AssetData) { return AssetData.PackageName; });
-
-		DiskTask.BusyWait();
-		TArray<FAssetData> DiskAssets = MoveTemp(DiskTask.GetResult());
-		ProduceAssets(DiskAssets, IgnorePackages);
+		}
 	}
 }
 
@@ -2342,8 +2258,7 @@ bool UContentBrowserAssetDataSource::CreateFolder(const FName InPath, FContentBr
 		InPath,
 		*FolderItemName,
 		FText::AsCultureInvariant(FolderItemName),
-		MakeShared<FContentBrowserAssetFolderItemDataPayload>(*InternalPathString),
-		FName(InternalPathString)
+		MakeShared<FContentBrowserAssetFolderItemDataPayload>(*InternalPathString)
 		);
 
 	OutPendingItem = FContentBrowserItemDataTemporaryContext(
diff --git a/Engine/Plugins/Editor/ContentBrowser/ContentBrowserAssetDataSource/Source/ContentBrowserAssetDataSource/Public/ContentBrowserAssetDataSource.h b/Engine/Plugins/Editor/ContentBrowser/ContentBrowserAssetDataSource/Source/ContentBrowserAssetDataSource/Public/ContentBrowserAssetDataSource.h
--- a/Engine/Plugins/Editor/ContentBrowser/ContentBrowserAssetDataSource/Source/ContentBrowserAssetDataSource/Public/ContentBrowserAssetDataSource.h
+++ b/Engine/Plugins/Editor/ContentBrowser/ContentBrowserAssetDataSource/Source/ContentBrowserAssetDataSource/Public/ContentBrowserAssetDataSource.h
@@ -222,7 +222,7 @@ public:
 	 *
 	 * @return 
 	 */
-	static void EnumerateFoldersMatchingFilter(UContentBrowserDataSource* DataSource, const FContentBrowserCompiledAssetDataFilter* AssetDataFilter, const TGetOrEnumerateSink<FContentBrowserItemData>& InSink, FSubPathEnumerationFunc SubPathEnumeration, FCreateFolderItemFunc CreateFolderItem);
+	static void EnumerateFoldersMatchingFilter(UContentBrowserDataSource* DataSource, const FContentBrowserCompiledAssetDataFilter* AssetDataFilter, TFunctionRef<bool(FContentBrowserItemData&&)> InCallback, FSubPathEnumerationFunc SubPathEnumeration, FCreateFolderItemFunc CreateFolderItem);
 	
 /**
 	 * Call in DoesItemPassFilter() to check if a folder passes the compiled asset data filter.
@@ -239,7 +239,6 @@ public:
 	virtual void CompileFilter(const FName InPath, const FContentBrowserDataFilter& InFilter, FContentBrowserDataCompiledFilter& OutCompiledFilter) override;
 
 	virtual void EnumerateItemsMatchingFilter(const FContentBrowserDataCompiledFilter& InFilter, TFunctionRef<bool(FContentBrowserItemData&&)> InCallback) override;
-	virtual void EnumerateItemsMatchingFilter(const FContentBrowserDataCompiledFilter& InFilter, const TGetOrEnumerateSink<FContentBrowserItemData>& InSink) override;
 
 	virtual void EnumerateItemsAtPath(const FName InPath, const EContentBrowserItemTypeFilter InItemTypeFilter, TFunctionRef<bool(FContentBrowserItemData&&)> InCallback) override;
 
diff --git a/Engine/Source/Editor/ContentBrowser/Private/AssetViewTypes.cpp b/Engine/Source/Editor/ContentBrowser/Private/AssetViewTypes.cpp
--- a/Engine/Source/Editor/ContentBrowser/Private/AssetViewTypes.cpp
+++ b/Engine/Source/Editor/ContentBrowser/Private/AssetViewTypes.cpp
@@ -65,15 +65,6 @@ void FAssetViewItem::RemoveItemData(const FContentBrowserItemData& InItemData)
 	}
 }
 
-void FAssetViewItem::RemoveItemData(const FContentBrowserMinimalItemData& InItemKey)
-{
-	Item.TryRemove(InItemKey);
-	if (Item.IsValid())
-	{
-		ItemDataChangedEvent.Broadcast();
-	}	
-}
-
 void FAssetViewItem::ClearCachedCustomColumns()
 {
 	CachedCustomColumnData.Reset();
diff --git a/Engine/Source/Editor/ContentBrowser/Private/PathViewTypes.cpp b/Engine/Source/Editor/ContentBrowser/Private/PathViewTypes.cpp
--- a/Engine/Source/Editor/ContentBrowser/Private/PathViewTypes.cpp
+++ b/Engine/Source/Editor/ContentBrowser/Private/PathViewTypes.cpp
@@ -58,11 +58,6 @@ void FTreeItem::RemoveItemData(const FContentBrowserItemData& InItemData)
 	Item.Remove(InItemData);
 }
 
-FContentBrowserItemData FTreeItem::RemoveItemData(const FContentBrowserMinimalItemData& InItemKey)
-{
-	return Item.TryRemove(InItemKey);
-}
-
 void FTreeItem::SetItemData(FContentBrowserItem InItem)
 {
 	Item = MoveTemp(InItem);
diff --git a/Engine/Source/Editor/ContentBrowser/Private/PathViewTypes.h b/Engine/Source/Editor/ContentBrowser/Private/PathViewTypes.h
--- a/Engine/Source/Editor/ContentBrowser/Private/PathViewTypes.h
+++ b/Engine/Source/Editor/ContentBrowser/Private/PathViewTypes.h
@@ -38,7 +38,6 @@ public:
 	void AppendItemData(const FContentBrowserItemData& InItemData);
 	void RemoveItemData(const FContentBrowserItem& InItem);
 	void RemoveItemData(const FContentBrowserItemData& InItemData);
-	FContentBrowserItemData RemoveItemData(const FContentBrowserMinimalItemData& InItemKey);
 
 	/** Get the underlying Content Browser item */
 	const FContentBrowserItem& GetItem() const;
diff --git a/Engine/Source/Editor/ContentBrowser/Private/SAssetView.cpp b/Engine/Source/Editor/ContentBrowser/Private/SAssetView.cpp
--- a/Engine/Source/Editor/ContentBrowser/Private/SAssetView.cpp
+++ b/Engine/Source/Editor/ContentBrowser/Private/SAssetView.cpp
@@ -2070,8 +2070,7 @@ void SAssetView::RefreshSourceItems()
 								*FString::Printf(TEXT("/Collections/%s/%s"), ECollectionShareType::ToString(ChildCollection.Type), *ChildCollection.Name.ToString()), 
 								ChildCollection.Name, 
 								FText::FromName(ChildCollection.Name), 
-								nullptr,
-								FName()
+								nullptr
 								);
 
 							const FContentBrowserItemKey FolderItemDataKey(FolderItemData);
@@ -4859,23 +4858,23 @@ void SAssetView::HandleItemDataUpdated(TArrayView<const FContentBrowserItemDataU
 		}
 	};
 
-	auto RemoveItem = [this, &bRefreshView, &ItemsPendingInplaceFrontendFilter](const FContentBrowserItemKey& ItemKey, const FContentBrowserMinimalItemData& ItemDataKey)
+	auto RemoveItem = [this, &bRefreshView, &ItemsPendingInplaceFrontendFilter](const FContentBrowserItemKey& InItemDataKey, const FContentBrowserItemData& InItemData)
 	{
-		const uint32 ItemDataKeyHash = GetTypeHash(ItemKey);
+		const uint32 ItemDataKeyHash = GetTypeHash(InItemDataKey);
 
-		if (const TSharedPtr<FAssetViewItem>* ItemToRemovePtr = AvailableBackendItems.FindByHash(ItemDataKeyHash, ItemKey))
+		if (const TSharedPtr<FAssetViewItem>* ItemToRemovePtr = AvailableBackendItems.FindByHash(ItemDataKeyHash, InItemDataKey))
 		{
 			TSharedPtr<FAssetViewItem> ItemToRemove = *ItemToRemovePtr;
 			check(ItemToRemove);
 
 			// Only fully remove this item if every sub-item is removed (items become invalid when empty)
-			ItemToRemove->RemoveItemData(ItemDataKey);
+			ItemToRemove->RemoveItemData(InItemData);
 			if (ItemToRemove->GetItem().IsValid())
 			{
 				return;
 			}
 
-			AvailableBackendItems.RemoveByHash(ItemDataKeyHash, ItemKey);
+			AvailableBackendItems.RemoveByHash(ItemDataKeyHash, InItemDataKey);
 
 			const uint32 ItemToRemoveHash = GetTypeHash(ItemToRemove);
 
@@ -4938,17 +4937,19 @@ void SAssetView::HandleItemDataUpdated(TArrayView<const FContentBrowserItemDataU
 			}
 			else
 			{
-				RemoveItem(ItemDataKey, FContentBrowserMinimalItemData(ItemData));
+				RemoveItem(ItemDataKey, ItemData);
 			}
 			break;
 
 		case EContentBrowserItemUpdateType::Moved:
 			{
-				const FContentBrowserMinimalItemData OldItemDataKey(ItemData.GetItemType(), ItemDataUpdate.GetPreviousVirtualPath(), ItemData.GetOwnerDataSource());
-				RemoveItem(ItemDataKey, OldItemDataKey);
+				const FContentBrowserItemData OldMinimalItemData(ItemData.GetOwnerDataSource(), ItemData.GetItemType(), ItemDataUpdate.GetPreviousVirtualPath(), NAME_None, FText(), nullptr);
+				const FContentBrowserItemKey OldItemDataKey(OldMinimalItemData);
+				RemoveItem(OldItemDataKey, OldMinimalItemData);
+
 				if (bItemPassFilter)
 				{
-					AddItem(FContentBrowserItemKey(ItemData), MoveTemp(ItemData));
+					AddItem(ItemDataKey, MoveTemp(ItemData));
 				}
 				else
 				{
@@ -4958,7 +4959,7 @@ void SAssetView::HandleItemDataUpdated(TArrayView<const FContentBrowserItemDataU
 			break;
 
 		case EContentBrowserItemUpdateType::Removed:
-			RemoveItem(ItemDataKey, FContentBrowserMinimalItemData(ItemData));
+			RemoveItem(ItemDataKey, ItemData);
 			break;
 
 		default:
diff --git a/Engine/Source/Editor/ContentBrowser/Private/SPathView.cpp b/Engine/Source/Editor/ContentBrowser/Private/SPathView.cpp
--- a/Engine/Source/Editor/ContentBrowser/Private/SPathView.cpp
+++ b/Engine/Source/Editor/ContentBrowser/Private/SPathView.cpp
@@ -339,7 +339,6 @@ protected:
 	// Remove the given item data from the tree - if this results in an item having no data from any sources, the tree
 	// item is removed. If the item was removed and had a parent, the parent is returned.
 	TSharedPtr<FTreeItem> TryRemoveFolderItemInternal(const FContentBrowserItemData& InItem);
-	TSharedPtr<FTreeItem> TryRemoveFolderItemInternal(const FContentBrowserMinimalItemData& InKey);
 
 	bool PassesTextFilter(const TSharedPtr<FTreeItem>& InItem);
 
@@ -532,8 +531,13 @@ void FPathViewData::ProcessDataUpdates(TConstArrayView<FContentBrowserItemDataUp
 
 			case EContentBrowserItemUpdateType::Moved:
 			{
-				const FContentBrowserMinimalItemData OldItemKey(ItemData.GetItemType(), ItemDataUpdate.GetPreviousVirtualPath(), ItemData.GetOwnerDataSource());
-				TSharedPtr<FTreeItem> Parent = TryRemoveFolderItemInternal(OldItemKey);
+				const FContentBrowserItemData OldMinimalItemData(ItemData.GetOwnerDataSource(),
+					ItemData.GetItemType(),
+					ItemDataUpdate.GetPreviousVirtualPath(),
+					NAME_None,
+					FText(),
+					nullptr);
+				TSharedPtr<FTreeItem> Parent = TryRemoveFolderItemInternal(OldMinimalItemData);
 				if (DoesItemPassFilter(ItemData))
 				{
 					NewItems.Emplace(AddFolderItemInternal(MoveTemp(ItemData), nullptr));
@@ -749,9 +753,7 @@ TSharedRef<FTreeItem> FPathViewData::AddFolderItemInternal(FContentBrowserItemDa
 																 ParentVirtualPath,
 																 FName(FPathViews::GetPathLeaf(PathView)),
 																 FText(),
-																 nullptr, 
-																 FName() // Assuming this is a virtual path with no internal path
-																 ),
+																 nullptr),
 					OldItemsByInvariantPath);
 				VirtualPathToItem.Add(ParentVirtualPath, ParentItem);
 				InvariantPathToItem.Add(
@@ -765,18 +767,19 @@ TSharedRef<FTreeItem> FPathViewData::AddFolderItemInternal(FContentBrowserItemDa
 	return LeafItem.ToSharedRef();
 }
 
-TSharedPtr<FTreeItem> FPathViewData::TryRemoveFolderItemInternal(const FContentBrowserItemData& InItemData)
+TSharedPtr<FTreeItem> FPathViewData::TryRemoveFolderItemInternal(const FContentBrowserItemData& InItem)
 {
-	return TryRemoveFolderItemInternal(FContentBrowserMinimalItemData(InItemData));
-}
+	if (!InItem.IsFolder())
+	{
+		// Not a folder
+		return {};
+	}
 
-TSharedPtr<FTreeItem> FPathViewData::TryRemoveFolderItemInternal(const FContentBrowserMinimalItemData& InItemKey)
-{
 	// Find the folder in the tree
-	if (TSharedPtr<FTreeItem> ItemToRemove = VirtualPathToItem.FindRef(InItemKey.GetVirtualPath()))
+	if (TSharedPtr<FTreeItem> ItemToRemove = VirtualPathToItem.FindRef(InItem.GetVirtualPath()))
 	{
 		// Only fully remove this item if every sub-item is removed (items become invalid when empty)
-		FContentBrowserItemData OldItemData = ItemToRemove->RemoveItemData(InItemKey);
+		ItemToRemove->RemoveItemData(InItem);
 		if (ItemToRemove->GetItem().IsValid())
 		{
 			return {};
@@ -796,8 +799,8 @@ TSharedPtr<FTreeItem> FPathViewData::TryRemoveFolderItemInternal(const FContentB
 			VisibleRootItems.Remove(ItemToRemove);
 		}
 
-		VirtualPathToItem.Remove(InItemKey.GetVirtualPath());
-		InvariantPathToItem.Remove(OldItemData.GetInvariantPath());
+		VirtualPathToItem.Remove(InItem.GetVirtualPath());
+		InvariantPathToItem.Remove(InItem.GetInvariantPath());
 		return ItemParent;
 	}
 
@@ -1091,12 +1094,8 @@ void SPathView::Construct( const FArguments& InArgs )
 
 	if (!InArgs._DefaultPath.IsEmpty() && InternalPathPassesBlockLists(InArgs._DefaultPath))
 	{
-		const FName VirtualPath = ContentBrowserData->ConvertInternalPathToVirtual(*InArgs._DefaultPath);
-		FName InternalPath;
-		if (ContentBrowserData->TryConvertVirtualPath(VirtualPath, InternalPath) != EContentBrowserPathType::Internal)
-		{
-			InternalPath = FName();
-		}
+		const FName VirtualPath =
+			IContentBrowserDataModule::Get().GetSubsystem()->ConvertInternalPathToVirtual(*InArgs._DefaultPath);
 		if (!TreeData->FindTreeItem(VirtualPath))
 		{
 			const FString DefaultPathLeafName = FPaths::GetPathLeaf(VirtualPath.ToString());
@@ -1105,8 +1104,7 @@ void SPathView::Construct( const FArguments& InArgs )
 				VirtualPath,
 				*DefaultPathLeafName,
 				FText(),
-				nullptr,
-				InternalPath));
+				nullptr));
 		}
 
 		SetSelectedPaths({ VirtualPath.ToString() });
diff --git a/Engine/Source/Editor/ContentBrowser/Public/AssetViewTypes.h b/Engine/Source/Editor/ContentBrowser/Public/AssetViewTypes.h
--- a/Engine/Source/Editor/ContentBrowser/Public/AssetViewTypes.h
+++ b/Engine/Source/Editor/ContentBrowser/Public/AssetViewTypes.h
@@ -36,8 +36,6 @@ public:
 
 	void RemoveItemData(const FContentBrowserItemData& InItemData);
 
-	void RemoveItemData(const FContentBrowserMinimalItemData& InItemKey);
-
 	/** Clear cached custom column data */
 	void ClearCachedCustomColumns();
 
diff --git a/Engine/Source/Editor/ContentBrowserData/Private/ContentBrowserDataSource.cpp b/Engine/Source/Editor/ContentBrowserData/Private/ContentBrowserDataSource.cpp
--- a/Engine/Source/Editor/ContentBrowserData/Private/ContentBrowserDataSource.cpp
+++ b/Engine/Source/Editor/ContentBrowserData/Private/ContentBrowserDataSource.cpp
@@ -218,21 +218,10 @@ void UContentBrowserDataSource::EnumerateItemsMatchingFilter(const FContentBrows
 {
 }
 
-void UContentBrowserDataSource::EnumerateItemsMatchingFilter(const FContentBrowserDataCompiledFilter& InFilter, const TGetOrEnumerateSink<FContentBrowserItemData>& InSink)
-{
-	// Call older API function if this wasn't overridden
-	EnumerateItemsMatchingFilter(InFilter, [&InSink](FContentBrowserItemData&& Item) { return InSink.ProduceItem(MoveTemp(Item)); });
-}
-
 void UContentBrowserDataSource::EnumerateItemsAtPath(const FName InPath, const EContentBrowserItemTypeFilter InItemTypeFilter, TFunctionRef<bool(FContentBrowserItemData&&)> InCallback)
 {
 }
 
-void UContentBrowserDataSource::EnumerateItemsAtPath(const FName InPath, const EContentBrowserItemTypeFilter InItemTypeFilter, const TGetOrEnumerateSink<FContentBrowserItemData>& InSink)
-{
-	EnumerateItemsAtPath(InPath, InItemTypeFilter, [&InSink](FContentBrowserItemData&& Item) { return InSink.ProduceItem(MoveTemp(Item)); });
-}
-
 bool UContentBrowserDataSource::EnumerateItemsAtPaths(const TArrayView<FContentBrowserItemPath> InPaths, const EContentBrowserItemTypeFilter InItemTypeFilter, TFunctionRef<bool(FContentBrowserItemData&&)> InCallback)
 {
 	return true;
@@ -592,9 +581,7 @@ FContentBrowserItemData UContentBrowserDataSource::CreateVirtualFolderItem(const
 		InFolderPath,
 		*FolderItemName,
 		MoveTemp(FolderDisplayNameOverride),
-		nullptr,
-		FName() // Virtual folders have no internal path 
-		);
+		nullptr);
 }
 
 #undef LOCTEXT_NAMESPACE
diff --git a/Engine/Source/Editor/ContentBrowserData/Private/ContentBrowserDataSubsystem.cpp b/Engine/Source/Editor/ContentBrowserData/Private/ContentBrowserDataSubsystem.cpp
--- a/Engine/Source/Editor/ContentBrowserData/Private/ContentBrowserDataSubsystem.cpp
+++ b/Engine/Source/Editor/ContentBrowserData/Private/ContentBrowserDataSubsystem.cpp
@@ -9,7 +9,6 @@
 #include "Editor.h"
 #include "Features/IModularFeatures.h"
 #include "Framework/Application/SlateApplication.h"
-#include "GetOrEnumerateSink.h"
 #include "HAL/IConsoleManager.h"
 #include "IContentBrowserDataModule.h"
 #include "Interfaces/IPluginManager.h"
@@ -302,12 +301,7 @@ void UContentBrowserDataSubsystem::EnumerateItemsMatchingFilter(const FContentBr
 
 void UContentBrowserDataSubsystem::EnumerateItemsMatchingFilter(const FContentBrowserDataCompiledFilter& InFilter, TFunctionRef<bool(FContentBrowserItemData&&)> InCallback) const
 {
-	EnumerateItemsMatchingFilter(InFilter, TGetOrEnumerateSink<FContentBrowserItemData>(InCallback));
-}
-
-void UContentBrowserDataSubsystem::EnumerateItemsMatchingFilter(const FContentBrowserDataCompiledFilter& InFilter, const TGetOrEnumerateSink<FContentBrowserItemData>& InSink) const
-{
-	for (const TPair<FName, UContentBrowserDataSource*>& ActiveDataSourcePair : ActiveDataSources)
+	for (const auto& ActiveDataSourcePair : ActiveDataSources)
 	{
 		UContentBrowserDataSource* DataSource = ActiveDataSourcePair.Value;
 
@@ -321,8 +315,7 @@ void UContentBrowserDataSubsystem::EnumerateItemsMatchingFilter(const FContentBr
 					check(EnumHasAnyFlags(InFilter.ItemTypeFilter, EContentBrowserItemTypeFilter::IncludeFolders));
 
 					const FString MountLeafName = FPackageName::GetShortName(MountRootPart);
-					FName InternalPath; // Virtual folders have no internal path 
-					InSink.ProduceItem(FContentBrowserItemData(DataSource, EContentBrowserItemFlags::Type_Folder, MountRootPart, *MountLeafName, FText(), nullptr, InternalPath));
+					InCallback(FContentBrowserItemData(DataSource, EContentBrowserItemFlags::Type_Folder, MountRootPart, *MountLeafName, FText(), nullptr));
 				}
 			}
 
@@ -334,13 +327,13 @@ void UContentBrowserDataSubsystem::EnumerateItemsMatchingFilter(const FContentBr
 					for (const auto& It : VirtualFolderFilter->CachedSubPaths)
 					{
 						// how do we skip over this item if not included (Engine Content, Engine Plugins, C++ Classes, etc..)
-						InSink.ProduceItem(FContentBrowserItemData(It.Value));
+						InCallback(FContentBrowserItemData(It.Value));
 					}
 				}
 			}
 		}
 
-		DataSource->EnumerateItemsMatchingFilter(InFilter, InSink);
+		DataSource->EnumerateItemsMatchingFilter(InFilter, InCallback);
 	}
 }
 
@@ -354,18 +347,14 @@ void UContentBrowserDataSubsystem::EnumerateItemsUnderPath(const FName InPath, c
 }
 
 void UContentBrowserDataSubsystem::EnumerateItemsUnderPath(const FName InPath, const FContentBrowserDataFilter& InFilter, TFunctionRef<bool(FContentBrowserItemData&&)> InCallback) const
-{
-	FContentBrowserDataCompiledFilter CompiledFilter;
-	CompileFilter(InPath, InFilter, CompiledFilter);
-	EnumerateItemsMatchingFilter(CompiledFilter, MoveTemp(InCallback));
-}
-
-void UContentBrowserDataSubsystem::EnumerateItemsUnderPath(const FName InPath, const FContentBrowserDataFilter& InFilter, const TGetOrEnumerateSink<FContentBrowserItemData>& InSink) const
 {
 	FContentBrowserDataCompiledFilter CompiledFilter;
 	CompileFilter(InPath, InFilter, CompiledFilter);
 
-	EnumerateItemsMatchingFilter(CompiledFilter, InSink);
+	EnumerateItemsMatchingFilter(CompiledFilter, [&InCallback](FContentBrowserItemData&& InItemData)
+	{
+		return InCallback(MoveTemp(InItemData));
+	});
 }
 
 TArray<FContentBrowserItem> UContentBrowserDataSubsystem::GetItemsUnderPath(const FName InPath, const FContentBrowserDataFilter& InFilter) const
@@ -409,7 +398,7 @@ void UContentBrowserDataSubsystem::EnumerateItemsAtPath(const FName InPath, cons
 void UContentBrowserDataSubsystem::EnumerateItemsAtPath(const FName InPath, const EContentBrowserItemTypeFilter InItemTypeFilter, TFunctionRef<bool(FContentBrowserItemData&&)> InCallback) const
 {
 	bool bHandledVirtualFolder = false;
-	for (const TPair<FName, UContentBrowserDataSource*>& ActiveDataSourcePair : ActiveDataSources)
+	for (const auto& ActiveDataSourcePair : ActiveDataSources)
 	{
 		UContentBrowserDataSource* DataSource = ActiveDataSourcePair.Value;
 		FName InternalPath;
diff --git a/Engine/Source/Editor/ContentBrowserData/Private/ContentBrowserItem.cpp b/Engine/Source/Editor/ContentBrowserData/Private/ContentBrowserItem.cpp
--- a/Engine/Source/Editor/ContentBrowserData/Private/ContentBrowserItem.cpp
+++ b/Engine/Source/Editor/ContentBrowserData/Private/ContentBrowserItem.cpp
@@ -278,32 +278,6 @@ bool FContentBrowserItem::TryRemove(const FContentBrowserItemData& InOther, FTex
 	return true;
 }
 
-FContentBrowserItemData FContentBrowserItem::TryRemove(const FContentBrowserMinimalItemData& InKey, FText* OutError)
-{
-	const FContentBrowserItemData* PrimaryItemData = GetPrimaryInternalItem();
-	if (!PrimaryItemData)
-	{
-		// We are empty, nothing to do (but not a failure)
-		return FContentBrowserItemData();
-	}
-
-	// Remove the other item from our list
-	int32 Index = ItemDataArray.IndexOfByPredicate([Key=InKey](const FContentBrowserItemData& Data) { return Key == FContentBrowserMinimalItemData(Data); });
-	if (Index != INDEX_NONE)
-	{
-		FContentBrowserItemData Data = MoveTemp(ItemDataArray[Index]);
-		ItemDataArray.RemoveAt(Index);
-		return MoveTemp(Data);
-	}
-	else
-	{
-		if (OutError)
-		{
-			*OutError = LOCTEXT("RemoveError_KeyMismatch", "No data matching given key");
-		}
-		return FContentBrowserItemData();
-	}
-}
 bool FContentBrowserItem::IsFolder() const
 {
 	const FContentBrowserItemData* PrimaryItemData = GetPrimaryInternalItem();
diff --git a/Engine/Source/Editor/ContentBrowserData/Private/ContentBrowserItemData.cpp b/Engine/Source/Editor/ContentBrowserData/Private/ContentBrowserItemData.cpp
--- a/Engine/Source/Editor/ContentBrowserData/Private/ContentBrowserItemData.cpp
+++ b/Engine/Source/Editor/ContentBrowserData/Private/ContentBrowserItemData.cpp
@@ -16,7 +16,7 @@ FContentBrowserItemData::FContentBrowserItemData(
 	FName InItemName,
 	FText InDisplayNameOverride,
 	TSharedPtr<const IContentBrowserItemDataPayload> InPayload,
-	FName InInternalPath)
+	TOptional<FName> InInternalPath)
 	: OwnerDataSource(InOwnerDataSource)
 	, ItemFlags(InItemFlags)
 	, VirtualPath(InVirtualPath)
@@ -26,10 +26,6 @@ FContentBrowserItemData::FContentBrowserItemData(
 	, Payload(MoveTemp(InPayload))
 {
 	checkf(IsFolder() != IsFile(), TEXT("Items must be either a folder or a file!"));
-	if (CachedDisplayName.IsEmpty())
-	{
-		CachedDisplayName = FText::AsCultureInvariant(ItemName.ToString());
-	}
 }
 
 bool FContentBrowserItemData::operator==(const FContentBrowserItemData& InOther) const
@@ -122,7 +118,24 @@ FName FContentBrowserItemData::GetInvariantPath() const
 
 FName FContentBrowserItemData::GetInternalPath() const
 {
-	return InternalPath;
+	if (InternalPath.IsSet())
+	{
+		return InternalPath.GetValue();
+	}
+	if (UContentBrowserDataSource* DataSource = OwnerDataSource.Get())
+	{
+		if (!VirtualPath.IsNone())
+		{
+			FName ConvertedPath;
+			if (DataSource->TryConvertVirtualPath(VirtualPath, ConvertedPath) == EContentBrowserPathType::Internal)
+			{
+				InternalPath = ConvertedPath;
+				return ConvertedPath;
+			}
+		}
+	}
+
+	return NAME_None;
 }
 
 FName FContentBrowserItemData::GetItemName() const
@@ -132,6 +145,10 @@ FName FContentBrowserItemData::GetItemName() const
 
 FText FContentBrowserItemData::GetDisplayName() const
 {
+	if (CachedDisplayName.IsEmpty())
+	{
+		CachedDisplayName = FText::AsCultureInvariant(ItemName.ToString());
+	}
 	return CachedDisplayName;
 }
 
@@ -324,28 +341,6 @@ FContentBrowserItemData FContentBrowserItemDataTemporaryContext::FinalizeItem(co
 	return OnFinalizeItem.Execute(ItemData, InProposedName, OutErrorMsg);
 }
 
-FContentBrowserMinimalItemData::FContentBrowserMinimalItemData(const FContentBrowserItemData& InItemData)
-	: FContentBrowserMinimalItemData(InItemData.GetItemType(), InItemData.GetVirtualPath(), InItemData.GetOwnerDataSource())
-{
-
-}
-
-FContentBrowserMinimalItemData::FContentBrowserMinimalItemData(EContentBrowserItemFlags InItemType, FName InVirtualPath, const UContentBrowserDataSource* InSource)
-	: ItemType(InItemType)
-	, VirtualPath(InVirtualPath)
-	, DataSource(InSource)
-{
-
-}
-
-FString FContentBrowserMinimalItemData::ToString() const
-{
-	return FString::Printf(TEXT("%s:%s:%s"), 
-		*VirtualPath.ToString(), 
-		EnumHasAllFlags(ItemType, EContentBrowserItemFlags::Type_File) ? TEXT("File") : TEXT("Folder"),
-		DataSource ? *DataSource->GetName() : TEXT("null")
-	);
-}
 
 FContentBrowserItemDataKey::FContentBrowserItemDataKey(const FContentBrowserItemData& InItemData)
 	: ItemType(InItemData.GetItemType())
diff --git a/Engine/Source/Editor/ContentBrowserData/Public/ContentBrowserDataSource.h b/Engine/Source/Editor/ContentBrowserData/Public/ContentBrowserDataSource.h
--- a/Engine/Source/Editor/ContentBrowserData/Public/ContentBrowserDataSource.h
+++ b/Engine/Source/Editor/ContentBrowserData/Public/ContentBrowserDataSource.h
@@ -198,10 +198,8 @@ public:
 	 *
 	 * @param InFilter The compiled filter used to find matching items.
 	 * @param InCallback The function to invoke for each matching item (return true to continue enumeration).
-	 * @param InSink Parameter provided to receive output as either a callback or an array.
 	 */
 	virtual void EnumerateItemsMatchingFilter(const FContentBrowserDataCompiledFilter& InFilter, TFunctionRef<bool(FContentBrowserItemData&&)> InCallback);
-	virtual void EnumerateItemsMatchingFilter(const FContentBrowserDataCompiledFilter& InFilter, const TGetOrEnumerateSink<FContentBrowserItemData>& InSink);
 
 	/**
 	 * Enumerate items that have the given virtual path, optionally filtering by type, and invoking the callback for each matching item.
@@ -210,10 +208,9 @@ public:
 	 * @param InPath The virtual path to find items for.
 	 * @param InItemTypeFilter The types of items we want to find.
 	 * @param InCallback The function to invoke for each matching item (return true to continue enumeration).
-	 * @param InSink Parameter provided to receive output as either a callback or an array.
 	 */
 	virtual void EnumerateItemsAtPath(const FName InPath, const EContentBrowserItemTypeFilter InItemTypeFilter, TFunctionRef<bool(FContentBrowserItemData&&)> InCallback);
-	virtual void EnumerateItemsAtPath(const FName InPath, const EContentBrowserItemTypeFilter InItemTypeFilter, const TGetOrEnumerateSink<FContentBrowserItemData>& InSink);
+
 	
 	/**
 	 * Enumerate the items (folders and/or files) that exist at the given content browser paths.
diff --git a/Engine/Source/Editor/ContentBrowserData/Public/ContentBrowserDataSubsystem.h b/Engine/Source/Editor/ContentBrowserData/Public/ContentBrowserDataSubsystem.h
--- a/Engine/Source/Editor/ContentBrowserData/Public/ContentBrowserDataSubsystem.h
+++ b/Engine/Source/Editor/ContentBrowserData/Public/ContentBrowserDataSubsystem.h
@@ -18,7 +18,6 @@
 #include "CoreMinimal.h"
 #include "Delegates/Delegate.h"
 #include "EditorSubsystem.h"
-#include "GetOrEnumerateSink.h"
 #include "HAL/Platform.h"
 #include "Misc/AssertionMacros.h"
 #include "Misc/NamePermissionList.h"
@@ -178,26 +177,14 @@ public:
 	/**
 	 * Enumerate the items (folders and/or files) that match a previously compiled filter.
 	 */
-	void EnumerateItemsMatchingFilter(
-		const FContentBrowserDataCompiledFilter& InFilter, TFunctionRef<bool(FContentBrowserItem&&)> InCallback) const;
-	void EnumerateItemsMatchingFilter(const FContentBrowserDataCompiledFilter& InFilter,
-		TFunctionRef<bool(FContentBrowserItemData&&)> InCallback) const;
-	void EnumerateItemsMatchingFilter(
-		const FContentBrowserDataCompiledFilter& InFilter, const TGetOrEnumerateSink<FContentBrowserItemData>& InSink) const;
+	void EnumerateItemsMatchingFilter(const FContentBrowserDataCompiledFilter& InFilter, TFunctionRef<bool(FContentBrowserItem&&)> InCallback) const;
+	void EnumerateItemsMatchingFilter(const FContentBrowserDataCompiledFilter& InFilter, TFunctionRef<bool(FContentBrowserItemData&&)> InCallback) const;
 
 	/**
 	 * Enumerate the items (folders and/or files) that exist under the given virtual path.
 	 */
-	void EnumerateItemsUnderPath(const FName InPath,
-		const FContentBrowserDataFilter& InFilter,
-		TFunctionRef<bool(FContentBrowserItem&&)> InCallback) const;
-	void EnumerateItemsUnderPath(const FName InPath,
-		const FContentBrowserDataFilter& InFilter,
-		TFunctionRef<bool(FContentBrowserItemData&&)> InCallback) const;
-	// Version which allows passing an array or a callback, allows sources to presize the array for large queries
-	void EnumerateItemsUnderPath(const FName InPath,
-		const FContentBrowserDataFilter& InFilter,
-		const TGetOrEnumerateSink<FContentBrowserItemData>& InSink) const;
+	void EnumerateItemsUnderPath(const FName InPath, const FContentBrowserDataFilter& InFilter, TFunctionRef<bool(FContentBrowserItem&&)> InCallback) const;
+	void EnumerateItemsUnderPath(const FName InPath, const FContentBrowserDataFilter& InFilter, TFunctionRef<bool(FContentBrowserItemData&&)> InCallback) const;
 
 	/**
 	 * Get the items (folders and/or files) that exist under the given virtual path.
diff --git a/Engine/Source/Editor/ContentBrowserData/Public/ContentBrowserItem.h b/Engine/Source/Editor/ContentBrowserData/Public/ContentBrowserItem.h
--- a/Engine/Source/Editor/ContentBrowserData/Public/ContentBrowserItem.h
+++ b/Engine/Source/Editor/ContentBrowserData/Public/ContentBrowserItem.h
@@ -20,7 +20,6 @@
 #include "ContentBrowserItem.generated.h"
 
 class FAssetThumbnail;
-class FContentBrowserMinimalItemData;
 class UContentBrowserDataSource;
 class UObject;
 struct FAssetData;
@@ -114,12 +113,6 @@ public:
 	 */
 	bool TryRemove(const FContentBrowserItemData& InOther, FText* OutError = nullptr);
 
-	/**
-	 * Attempt to remove an item matching the given minimal representation from this one, providing error information if the removal wasn't possible.
-	 * @return A valid item data (IsValid) if anything was removed, an invalid one otherwise. 
-	 */
-	FContentBrowserItemData TryRemove(const FContentBrowserMinimalItemData& InKey, FText* OutError = nullptr);
-
 	/**
 	 * Get the array of internal items that comprise this composite item.
 	 */
diff --git a/Engine/Source/Editor/ContentBrowserData/Public/ContentBrowserItemData.h b/Engine/Source/Editor/ContentBrowserData/Public/ContentBrowserItemData.h
--- a/Engine/Source/Editor/ContentBrowserData/Public/ContentBrowserItemData.h
+++ b/Engine/Source/Editor/ContentBrowserData/Public/ContentBrowserItemData.h
@@ -19,7 +19,6 @@
 #include "UObject/WeakObjectPtr.h"
 #include "UObject/WeakObjectPtrTemplates.h"
 
-class FContentBrowserMinimalItemData;
 class IContentBrowserItemDataSink;
 class UContentBrowserDataSource;
 
@@ -135,7 +134,7 @@ public:
 	 * @param InDisplayNameOverride The user-facing name of this item (eg, "MyFile"). This will be lazily set to
 	 * InItemName if no override is provided.
 	 * @param InPayload Any data source defined payload data for this item.
-	 * @param InInternalPath Internal path if any (some items such as "virtual folders" may not have an internal path)
+	 * @param InInternalPath Precalculated internal path if available
 	 */
 	FContentBrowserItemData(
 		UContentBrowserDataSource* InOwnerDataSource,
@@ -144,7 +143,7 @@ public:
 		FName InItemName,
 		FText InDisplayNameOverride,
 		TSharedPtr<const IContentBrowserItemDataPayload> InPayload,
-		FName InInternalPath);
+		TOptional<FName> InInternalPath = {});
 
 	/**
 	 * Copy support.
@@ -279,11 +278,11 @@ private:
 	/** The leaf-name of this item (eg, "MyFile") */
 	FName ItemName;
 
-	/** Internal path used to identify item within the owner data source */
-	FName InternalPath;
+	/** Cached/pre-computed value for GetInternalPath */
+	mutable TOptional<FName> InternalPath;
 
 	/** The user-facing name of this item (eg, "MyFile") */
-	FText CachedDisplayName;
+	mutable FText CachedDisplayName;
 
 	/** Any data source defined payload data for this item */
 	TSharedPtr<const IContentBrowserItemDataPayload> Payload;
@@ -625,60 +624,6 @@ private:
 	FOnFinalizeItem OnFinalizeItem;
 };
 
-/**
- * Minimal representation of a FContentBrowserItemData instance for comparison purposes.
- * @note Unlike FContentBrowserItemDataKey this does consider the data source, and unlike FContentBrowserItemKey it considers the data source for folder items.
- *  It can be used to identify a FContentBrowserItemData for removal from an FContentBrowserItem
- */
-class CONTENTBROWSERDATA_API FContentBrowserMinimalItemData
-{
-public:
-	explicit FContentBrowserMinimalItemData(const FContentBrowserItemData& InItemData);
-	FContentBrowserMinimalItemData(EContentBrowserItemFlags InItemType, FName InVirtualPath, const UContentBrowserDataSource* InSource);
-
-	bool operator==(const FContentBrowserMinimalItemData& Other) const
-	{
-		return !(*this != Other);
-	}
-	bool operator!=(const FContentBrowserMinimalItemData& Other) const
-	{
-		return ItemType == Other.ItemType && VirtualPath == Other.VirtualPath && DataSource == Other.DataSource;
-	}
-
-	/** Return a string representation for debugging. */
-	FString ToString() const;
-	
-	EContentBrowserItemFlags GetItemType() const
-	{
-		return ItemType;
-	} 
-
-	FName GetVirtualPath() const
-	{
-		return VirtualPath;
-	}
-
-	const UContentBrowserDataSource* GetDataSource() const 
-	{
-		return DataSource;
-	}
-
-private:
-	/**
-	 * Flags denoting the item type information for an item instance.
-	 * @note This is always masked against EContentBrowserItemFlags::Type_MASK to remove any non-type information.
-	 */
-	EContentBrowserItemFlags ItemType = EContentBrowserItemFlags::None;
-
-	/**
-	 * The complete virtual path that uniquely identifies an item within its owner data source.
-	 */
-	FName VirtualPath;
-
-	/** A pointer to the data source that manages the thing represented by this key */
-	const UContentBrowserDataSource* DataSource = nullptr;
-};
-
 /**
  * Minimal representation of a FContentBrowserItemData instance that can be used as a map key.
  * @note This key doesn't consider the data source, so should only be used for maps within a given data source.
diff --git a/Engine/Source/Editor/ContentBrowserData/Public/GetOrEnumerateSink.h b/Engine/Source/Editor/ContentBrowserData/Public/GetOrEnumerateSink.h
deleted file mode 100644
