Add RWLock for cache access. Add method to return cached data if available without population.

diff --git a/Engine/Plugins/Editor/ContentBrowser/ContentBrowserAssetDataSource/Source/ContentBrowserAssetDataSource/Private/AssetPropertyTagCache.cpp b/Engine/Plugins/Editor/ContentBrowser/ContentBrowserAssetDataSource/Source/ContentBrowserAssetDataSource/Private/AssetPropertyTagCache.cpp
--- a/Engine/Plugins/Editor/ContentBrowser/ContentBrowserAssetDataSource/Source/ContentBrowserAssetDataSource/Private/AssetPropertyTagCache.cpp
+++ b/Engine/Plugins/Editor/ContentBrowser/ContentBrowserAssetDataSource/Source/ContentBrowserAssetDataSource/Private/AssetPropertyTagCache.cpp
@@ -11,12 +11,26 @@ FAssetPropertyTagCache& FAssetPropertyTagCache::Get()
 	return AssetPropertyTagCache;
 }
 
-const FAssetPropertyTagCache::FClassPropertyTagCache& FAssetPropertyTagCache::GetCacheForClass(FTopLevelAssetPath InClassName)
+const FAssetPropertyTagCache::FClassPropertyTagCache* FAssetPropertyTagCache::FindCacheForClass(FTopLevelAssetPath InClassName)
 {
+	FReadScopeLock ReadScope(Lock);
 	TSharedPtr<FClassPropertyTagCache> ClassCache = ClassToCacheMap.FindRef(InClassName);
+	return ClassCache.Get();
+}
+
+const FAssetPropertyTagCache::FClassPropertyTagCache& FAssetPropertyTagCache::GetCacheForClass(FTopLevelAssetPath InClassName)
+{
+	TSharedPtr<FClassPropertyTagCache> ClassCache;
+	
+	{
+		FReadScopeLock ReadScope(Lock);
+		ClassCache = ClassToCacheMap.FindRef(InClassName);
+	}
+	
+	// On a miss, race with other potential creations rather than hold the lock against queries for another class
 	if (!ClassCache)
 	{
-		ClassCache = ClassToCacheMap.Add(InClassName, MakeShared<FClassPropertyTagCache>());
+		ClassCache = MakeShared<FClassPropertyTagCache>();
 
 		auto GetAssetClass = [&InClassName]()
 		{
@@ -91,6 +105,9 @@ const FAssetPropertyTagCache::FClassPropertyTagCache& FAssetPropertyTagCache::Ge
 				}
 			}
 		}
+
+		FWriteScopeLock WriteScope(Lock);
+		ClassToCacheMap.Add(InClassName, ClassCache);
 	}
 
 	check(ClassCache);
diff --git a/Engine/Plugins/Editor/ContentBrowser/ContentBrowserAssetDataSource/Source/ContentBrowserAssetDataSource/Public/AssetPropertyTagCache.h b/Engine/Plugins/Editor/ContentBrowser/ContentBrowserAssetDataSource/Source/ContentBrowserAssetDataSource/Public/AssetPropertyTagCache.h
--- a/Engine/Plugins/Editor/ContentBrowser/ContentBrowserAssetDataSource/Source/ContentBrowserAssetDataSource/Public/AssetPropertyTagCache.h
+++ b/Engine/Plugins/Editor/ContentBrowser/ContentBrowserAssetDataSource/Source/ContentBrowserAssetDataSource/Public/AssetPropertyTagCache.h
@@ -61,8 +61,16 @@ public:
 	/** Get (or populate) the cache for the given asset class */
 	const FClassPropertyTagCache& GetCacheForClass(FTopLevelAssetPath InClassName);
 
+	/** Get the cache for the given asset class if it has been created */
+	const FClassPropertyTagCache* FindCacheForClass(FTopLevelAssetPath InClassName);
+
 private:
-	/** Mapping of the asset class name to its cache */
+	mutable FRWLock Lock;
+
+	/** 
+	 * Mapping of the asset class name to its cache. 
+	 * Values are not modified after construction, so they can safely be returned after releasing the lock. 
+	 */
 	TMap<FTopLevelAssetPath, TSharedPtr<FClassPropertyTagCache>> ClassToCacheMap;
 };
 
