Build information about asset registry tags on the game thread during asset registry data gather to allow reading it from other threads in content browser code.

diff --git a/Engine/Plugins/Editor/ContentBrowser/ContentBrowserAssetDataSource/Source/ContentBrowserAssetDataSource/Private/ContentBrowserAssetDataCore.cpp b/Engine/Plugins/Editor/ContentBrowser/ContentBrowserAssetDataSource/Source/ContentBrowserAssetDataSource/Private/ContentBrowserAssetDataCore.cpp
--- a/Engine/Plugins/Editor/ContentBrowser/ContentBrowserAssetDataSource/Source/ContentBrowserAssetDataSource/Private/ContentBrowserAssetDataCore.cpp
+++ b/Engine/Plugins/Editor/ContentBrowser/ContentBrowserAssetDataSource/Source/ContentBrowserAssetDataSource/Private/ContentBrowserAssetDataCore.cpp
@@ -1529,7 +1529,7 @@ bool GetVirtualizationItemAttribute(const FAssetData& InAssetData, IAssetRegistr
 	}
 }
 
-void GetGenericItemAttribute(const FName InTagKey, const FString& InTagValue, const FAssetPropertyTagCache::FClassPropertyTagCache& InClassPropertyTagCache, const bool InIncludeMetaData, FContentBrowserItemDataAttributeValue& OutAttributeValue)
+void GetGenericItemAttribute(const FName InTagKey, const FString& InTagValue, const FAssetPropertyTagCache::FClassPropertyTagCache* InClassPropertyTagCache, const bool InIncludeMetaData, FContentBrowserItemDataAttributeValue& OutAttributeValue)
 {
 	check(!InTagKey.IsNone());
 
@@ -1549,7 +1549,8 @@ void GetGenericItemAttribute(const FName InTagKey, const FString& InTagValue, co
 	if (InIncludeMetaData)
 	{
 		FContentBrowserItemDataAttributeMetaData AttributeMetaData;
-		if (const FAssetPropertyTagCache::FPropertyTagCache* PropertyTagCache = InClassPropertyTagCache.GetCacheForTag(InTagKey))
+		const FAssetPropertyTagCache::FPropertyTagCache* PropertyTagCache = InClassPropertyTagCache ? InClassPropertyTagCache->GetCacheForTag(InTagKey) : nullptr;
+		if (PropertyTagCache)
 		{
 			AttributeMetaData.AttributeType = PropertyTagCache->TagType;
 			AttributeMetaData.DisplayFlags = PropertyTagCache->DisplayFlags;
@@ -1785,17 +1786,20 @@ bool GetAssetDataAttribute(const FAssetData& InAssetData, const bool InIncludeMe
 
 	// Generic attribute keys
 	{
-		const FAssetPropertyTagCache::FClassPropertyTagCache& ClassPropertyTagCache = FAssetPropertyTagCache::Get().GetCacheForClass(InAssetData.AssetClassPath);
+		const FAssetPropertyTagCache::FClassPropertyTagCache* ClassPropertyTagCache = FAssetPropertyTagCache::Get().FindCacheForClass(InAssetData.AssetClassPath);
 
 		FName FoundAttributeKey = InAttributeKey;
 		FAssetDataTagMapSharedView::FFindTagResult FoundValue = InAssetData.TagsAndValues.FindTag(FoundAttributeKey);
 		if (!FoundValue.IsSet())
 		{
-			// Check to see if the key we were given resolves as an alias
-			FoundAttributeKey = ClassPropertyTagCache.GetTagNameFromAlias(FoundAttributeKey);
-			if (!FoundAttributeKey.IsNone())
+			if (ensureMsgf(ClassPropertyTagCache, TEXT("FAssetPropertyTagCache not populated for type %s"), *WriteToString<256>(InAssetData.AssetClassPath)))
 			{
-				FoundValue = InAssetData.TagsAndValues.FindTag(FoundAttributeKey);
+				// Check to see if the key we were given resolves as an alias
+				FoundAttributeKey = ClassPropertyTagCache->GetTagNameFromAlias(FoundAttributeKey);
+				if (!FoundAttributeKey.IsNone())
+				{
+					FoundValue = InAssetData.TagsAndValues.FindTag(FoundAttributeKey);
+				}
 			}
 		}
 		if (FoundValue.IsSet())
@@ -1865,7 +1869,7 @@ bool GetAssetDataAttributes(const FAssetData& InAssetData, const bool InIncludeM
 	static const FTopLevelAssetPath BlueprintAssetClass = FTopLevelAssetPath(TEXT("/Script/Engine"), TEXT("Blueprint"));
 	static const FName ParentClassTag = FName("ParentClass");
 	{
-		const FAssetPropertyTagCache::FClassPropertyTagCache& ClassPropertyTagCache = FAssetPropertyTagCache::Get().GetCacheForClass(InAssetData.AssetClassPath);
+		const FAssetPropertyTagCache::FClassPropertyTagCache* ClassPropertyTagCache = FAssetPropertyTagCache::Get().FindCacheForClass(InAssetData.AssetClassPath);
 		const FAssetPropertyTagCache::FClassPropertyTagCache* ParentClassPropertyTagCache = nullptr;
 
 		if (InAssetData.AssetClassPath == BlueprintAssetClass)
@@ -1877,7 +1881,7 @@ bool GetAssetDataAttributes(const FAssetData& InAssetData, const bool InIncludeM
 				FTopLevelAssetPath ParentClassPathName = UClass::TryConvertShortTypeNameToPathName<UClass>(ParentClassName, ELogVerbosity::Warning, TEXT("GetAssetFileItemAttributes"));
 				if (!ParentClassPathName.IsNull())
 				{
-					ParentClassPropertyTagCache = &FAssetPropertyTagCache::Get().GetCacheForClass(ParentClassPathName);
+					ParentClassPropertyTagCache = FAssetPropertyTagCache::Get().FindCacheForClass(ParentClassPathName);
 				}
 				else
 				{
@@ -1893,7 +1897,7 @@ bool GetAssetDataAttributes(const FAssetData& InAssetData, const bool InIncludeM
 			GetGenericItemAttribute(TagAndValue.Key, TagAndValue.Value.AsString(), ClassPropertyTagCache, InIncludeMetaData, GenericAttributeValue);
 			if (ParentClassPropertyTagCache && ParentClassPropertyTagCache->GetCacheForTag(TagAndValue.Key))
 			{
-				GetGenericItemAttribute(TagAndValue.Key, TagAndValue.Value.AsString(), *ParentClassPropertyTagCache, InIncludeMetaData, GenericAttributeValue);
+				GetGenericItemAttribute(TagAndValue.Key, TagAndValue.Value.AsString(), ParentClassPropertyTagCache, InIncludeMetaData, GenericAttributeValue);
 			}
 		}
 	}
diff --git a/Engine/Plugins/Editor/ContentBrowser/ContentBrowserAssetDataSource/Source/ContentBrowserAssetDataSource/Private/ContentBrowserAssetDataSource.cpp b/Engine/Plugins/Editor/ContentBrowser/ContentBrowserAssetDataSource/Source/ContentBrowserAssetDataSource/Private/ContentBrowserAssetDataSource.cpp
--- a/Engine/Plugins/Editor/ContentBrowser/ContentBrowserAssetDataSource/Source/ContentBrowserAssetDataSource/Private/ContentBrowserAssetDataSource.cpp
+++ b/Engine/Plugins/Editor/ContentBrowser/ContentBrowserAssetDataSource/Source/ContentBrowserAssetDataSource/Private/ContentBrowserAssetDataSource.cpp
@@ -3,6 +3,7 @@
 #include "ContentBrowserAssetDataSource.h"
 
 #include "Algo/Transform.h"
+#include "AssetPropertyTagCache.h"
 #include "Async/ParallelFor.h"
 #include "ContentBrowserAssetDataCore.h"
 #include "AssetRegistry/AssetRegistryModule.h"
@@ -155,7 +156,7 @@ void UContentBrowserAssetDataSource::Initialize(const bool InAutoRegister)
 	CollectionManager = &FCollectionManagerModule::GetModule().Get();
 
 	// Listen for asset registry updates
-	AssetRegistry->OnAssetAdded().AddUObject(this, &UContentBrowserAssetDataSource::OnAssetAdded);
+	AssetRegistry->OnAssetsAdded().AddUObject(this, &UContentBrowserAssetDataSource::OnAssetsAdded);
 	AssetRegistry->OnAssetRemoved().AddUObject(this, &UContentBrowserAssetDataSource::OnAssetRemoved);
 	AssetRegistry->OnAssetRenamed().AddUObject(this, &UContentBrowserAssetDataSource::OnAssetRenamed);
 	AssetRegistry->OnAssetUpdated().AddUObject(this, &UContentBrowserAssetDataSource::OnAssetUpdated);
@@ -245,6 +246,8 @@ void UContentBrowserAssetDataSource::Initialize(const bool InAutoRegister)
 
 	DiscoveryStatusText = LOCTEXT("InitializingAssetDiscovery", "Initializing Asset Discovery...");
 
+	FAssetPropertyTagCache& PropertyTagCache = FAssetPropertyTagCache::Get();
+
 	// Populate the initial set of folder attributes
 	// This will be updated as the scan finds more content
 	AssetRegistry->EnumerateAllCachedPaths([this](FName PathName) { 
@@ -252,8 +255,9 @@ void UContentBrowserAssetDataSource::Initialize(const bool InAutoRegister)
 		OnPathsAdded({NameBuilder.ToView()});
 		return true; 
 	});
-	AssetRegistry->EnumerateAllAssets([this](const FAssetData& InAssetData)
+	AssetRegistry->EnumerateAllAssets([this, &PropertyTagCache ](const FAssetData& InAssetData)
 		{
+			PropertyTagCache.GetCacheForClass(InAssetData.AssetClassPath);
 			OnPathPopulated(InAssetData);
 			return true;
 		}, /*bIncludeOnlyOnDiskAssets*/true);
@@ -290,7 +294,7 @@ void UContentBrowserAssetDataSource::Shutdown()
 		{
 			AssetRegistryMaybe->OnFileLoadProgressUpdated().RemoveAll(this);
 
-			AssetRegistryMaybe->OnAssetAdded().RemoveAll(this);
+			AssetRegistryMaybe->OnAssetsAdded().RemoveAll(this);
 			AssetRegistryMaybe->OnAssetRemoved().RemoveAll(this);
 			AssetRegistryMaybe->OnAssetRenamed().RemoveAll(this);
 			AssetRegistryMaybe->OnAssetUpdated().RemoveAll(this);
@@ -3135,14 +3139,22 @@ void UContentBrowserAssetDataSource::OnAssetRegistryFileLoadProgress(const IAsse
 	}
 }
 
-void UContentBrowserAssetDataSource::OnAssetAdded(const FAssetData& InAssetData)
+void UContentBrowserAssetDataSource::OnAssetsAdded(TConstArrayView<FAssetData> InAssets)
 {
-	if (ContentBrowserAssetData::IsPrimaryAsset(InAssetData))
+	TRACE_CPUPROFILER_EVENT_SCOPE(UContentBrowserAssetDataSource::OnAssetsAdded);
+
+	FAssetPropertyTagCache& Cache = FAssetPropertyTagCache::Get();
+	for (const FAssetData& InAssetData : InAssets)
 	{
-		// The owner folder of this asset is no longer considered empty
-		OnPathPopulated(InAssetData);
+		Cache.GetCacheForClass(InAssetData.AssetClassPath);
 
-		QueueItemDataUpdate(FContentBrowserItemDataUpdate::MakeItemAddedUpdate(CreateAssetFileItem(InAssetData)));
+		if (ContentBrowserAssetData::IsPrimaryAsset(InAssetData))
+		{
+			// The owner folder of this asset is no longer considered empty
+			OnPathPopulated(InAssetData);
+
+			QueueItemDataUpdate(FContentBrowserItemDataUpdate::MakeItemAddedUpdate(CreateAssetFileItem(InAssetData)));
+		}
 	}
 }
 
diff --git a/Engine/Plugins/Editor/ContentBrowser/ContentBrowserAssetDataSource/Source/ContentBrowserAssetDataSource/Public/ContentBrowserAssetDataSource.h b/Engine/Plugins/Editor/ContentBrowser/ContentBrowserAssetDataSource/Source/ContentBrowserAssetDataSource/Public/ContentBrowserAssetDataSource.h
--- a/Engine/Plugins/Editor/ContentBrowser/ContentBrowserAssetDataSource/Source/ContentBrowserAssetDataSource/Public/ContentBrowserAssetDataSource.h
+++ b/Engine/Plugins/Editor/ContentBrowser/ContentBrowserAssetDataSource/Source/ContentBrowserAssetDataSource/Public/ContentBrowserAssetDataSource.h
@@ -377,7 +377,7 @@ private:
 
 	void OnAssetRegistryFileLoadProgress(const IAssetRegistry::FFileLoadProgressUpdateData& InProgressUpdateData);
 
-	void OnAssetAdded(const FAssetData& InAssetData);
+	void OnAssetsAdded(TConstArrayView<FAssetData> InAssets);
 
 	void OnAssetRemoved(const FAssetData& InAssetData);
 
