Populate asset property tag cache for redirected class names.
Separate attempted cache population (which may be called for unavailable/unloaded classes) from GetCacheForClass so the latter can properly assert on failure.

diff --git a/Engine/Plugins/Editor/ContentBrowser/ContentBrowserAssetDataSource/Source/ContentBrowserAssetDataSource/Private/AssetPropertyTagCache.cpp b/Engine/Plugins/Editor/ContentBrowser/ContentBrowserAssetDataSource/Source/ContentBrowserAssetDataSource/Private/AssetPropertyTagCache.cpp
--- a/Engine/Plugins/Editor/ContentBrowser/ContentBrowserAssetDataSource/Source/ContentBrowserAssetDataSource/Private/AssetPropertyTagCache.cpp
+++ b/Engine/Plugins/Editor/ContentBrowser/ContentBrowserAssetDataSource/Source/ContentBrowserAssetDataSource/Private/AssetPropertyTagCache.cpp
@@ -1,6 +1,7 @@
 // Copyright Epic Games, Inc. All Rights Reserved.
 
 #include "AssetPropertyTagCache.h"
+#include "ContentBrowserAssetDataCore.h"
 #include "UObject/AssetRegistryTagsContext.h"
 #include "UObject/UnrealType.h"
 #include "UObject/LinkerLoad.h"
@@ -19,6 +20,26 @@ const FAssetPropertyTagCache::FClassPropertyTagCache* FAssetPropertyTagCache::Fi
 }
 
 const FAssetPropertyTagCache::FClassPropertyTagCache& FAssetPropertyTagCache::GetCacheForClass(FTopLevelAssetPath InClassName)
+{
+	TSharedPtr<FClassPropertyTagCache> ClassCache;
+	{
+		FReadScopeLock ReadScope(Lock);
+		ClassCache = ClassToCacheMap.FindRef(InClassName);
+	}
+	if (ClassCache.IsValid())
+	{
+		return *ClassCache;
+	}
+	TryCacheClass(InClassName);
+	{
+		FReadScopeLock ReadScope(Lock);
+		ClassCache = ClassToCacheMap.FindRef(InClassName);
+	}
+	checkf(ClassCache.IsValid(), TEXT("Failed to create property tag cache for class %s"), *WriteToString<256>(InClassName));
+	return *ClassCache;
+}
+
+void FAssetPropertyTagCache::TryCacheClass(FTopLevelAssetPath InClassName)
 {
 	TSharedPtr<FClassPropertyTagCache> ClassCache;
 	
@@ -26,90 +47,96 @@ const FAssetPropertyTagCache::FClassPropertyTagCache& FAssetPropertyTagCache::Ge
 		FReadScopeLock ReadScope(Lock);
 		ClassCache = ClassToCacheMap.FindRef(InClassName);
 	}
+
+	if (ClassCache.IsValid())
+	{
+		return;
+	}
 	
 	// On a miss, race with other potential creations rather than hold the lock against queries for another class
-	if (!ClassCache)
+	UE_LOG(LogContentBrowserAssetDataSource, Verbose, TEXT("Creating asset property tag cache for %s"), *WriteToString<256>(InClassName));
+
+	auto GetAssetClass = [&InClassName]()
 	{
-		ClassCache = MakeShared<FClassPropertyTagCache>();
+		UClass* FoundClass = FindObject<UClass>(InClassName);
 
-		auto GetAssetClass = [&InClassName]()
+		if (!FoundClass)
 		{
-			UClass* FoundClass = FindObject<UClass>(InClassName);
-
-			if (!FoundClass)
+			// Look for class redirectors
+			FString NewPath = FLinkerLoad::FindNewPathNameForClass(InClassName.ToString(), false);
+			if (!NewPath.IsEmpty())
 			{
-				// Look for class redirectors
-				FString NewPath = FLinkerLoad::FindNewPathNameForClass(InClassName.ToString(), false);
-				if (!NewPath.IsEmpty())
-				{
-					FoundClass = FindObject<UClass>(nullptr, *NewPath);
-				}
+				FoundClass = FindObject<UClass>(nullptr, *NewPath);
 			}
+		}
 
-			return FoundClass;
-		};
+		return FoundClass;
+	};
+
+	if (UClass* AssetClass = GetAssetClass())
+	{
+		ClassCache = MakeShared<FClassPropertyTagCache>();
+		FTopLevelAssetPath ActualClassPath(AssetClass);
 
-		if (UClass* AssetClass = GetAssetClass())
+		// Get the tags data from the CDO and use it to build the cache
+		UObject* AssetClassDefaultObject = AssetClass->GetDefaultObject();
+		FAssetRegistryTagsContextData TagsContext(AssetClassDefaultObject, EAssetRegistryTagsCaller::Uncategorized);
+		AssetClassDefaultObject->GetAssetRegistryTags(TagsContext);
+		TMap<FName, UObject::FAssetRegistryTagMetadata> AssetTagMetaData;
+		AssetClassDefaultObject->GetAssetRegistryTagMetadata(AssetTagMetaData);
+
+		ClassCache->TagNameToCachedDataMap.Reserve(TagsContext.Tags.Num());
+		for (const TPair<FName, UObject::FAssetRegistryTag>& TagPair : TagsContext.Tags)
 		{
-			// Get the tags data from the CDO and use it to build the cache
-			UObject* AssetClassDefaultObject = AssetClass->GetDefaultObject();
-			FAssetRegistryTagsContextData TagsContext(AssetClassDefaultObject, EAssetRegistryTagsCaller::Uncategorized);
-			AssetClassDefaultObject->GetAssetRegistryTags(TagsContext);
-			TMap<FName, UObject::FAssetRegistryTagMetadata> AssetTagMetaData;
-			AssetClassDefaultObject->GetAssetRegistryTagMetadata(AssetTagMetaData);
-
-			ClassCache->TagNameToCachedDataMap.Reserve(TagsContext.Tags.Num());
-			for (const TPair<FName, UObject::FAssetRegistryTag>& TagPair : TagsContext.Tags)
+			FPropertyTagCache& TagCache = ClassCache->TagNameToCachedDataMap.Add(TagPair.Key);
+			TagCache.TagType = TagPair.Value.Type;
+			TagCache.DisplayFlags = TagPair.Value.DisplayFlags;
+
+			if (const UObject::FAssetRegistryTagMetadata* TagMetaData = AssetTagMetaData.Find(TagPair.Key))
 			{
-				FPropertyTagCache& TagCache = ClassCache->TagNameToCachedDataMap.Add(TagPair.Key);
-				TagCache.TagType = TagPair.Value.Type;
-				TagCache.DisplayFlags = TagPair.Value.DisplayFlags;
+				TagCache.DisplayName = TagMetaData->DisplayName;
+				TagCache.TooltipText = TagMetaData->TooltipText;
+				TagCache.Suffix = TagMetaData->Suffix;
+				TagCache.ImportantValue = TagMetaData->ImportantValue;
+			}
+			else
+			{
+				// If the tag name corresponds to a property name, use the property tooltip
+				const FProperty* Property = FindFProperty<FProperty>(AssetClass, TagPair.Key);
+				TagCache.TooltipText = Property ? Property->GetToolTipText() : FText::FromString(FName::NameToDisplayString(TagPair.Key.ToString(), false));
+			}
 
-				if (const UObject::FAssetRegistryTagMetadata* TagMetaData = AssetTagMetaData.Find(TagPair.Key))
+			// Ensure a display name for this tag
+			if (TagCache.DisplayName.IsEmpty())
+			{
+				if (const FProperty* TagField = FindFProperty<FProperty>(AssetClass, TagPair.Key))
 				{
-					TagCache.DisplayName = TagMetaData->DisplayName;
-					TagCache.TooltipText = TagMetaData->TooltipText;
-					TagCache.Suffix = TagMetaData->Suffix;
-					TagCache.ImportantValue = TagMetaData->ImportantValue;
+					// Take the display name from the corresponding property if possible
+					TagCache.DisplayName = TagField->GetDisplayNameText();
 				}
 				else
 				{
-					// If the tag name corresponds to a property name, use the property tooltip
-					const FProperty* Property = FindFProperty<FProperty>(AssetClass, TagPair.Key);
-					TagCache.TooltipText = Property ? Property->GetToolTipText() : FText::FromString(FName::NameToDisplayString(TagPair.Key.ToString(), false));
-				}
-
-				// Ensure a display name for this tag
-				if (TagCache.DisplayName.IsEmpty())
-				{
-					if (const FProperty* TagField = FindFProperty<FProperty>(AssetClass, TagPair.Key))
-					{
-						// Take the display name from the corresponding property if possible
-						TagCache.DisplayName = TagField->GetDisplayNameText();
-					}
-					else
-					{
-						// We have no type information by this point, so no idea if it's a bool :(
-						TagCache.DisplayName = FText::AsCultureInvariant(FName::NameToDisplayString(TagPair.Key.ToString(), /*bIsBool*/false));
-					}
+					// We have no type information by this point, so no idea if it's a bool :(
+					TagCache.DisplayName = FText::AsCultureInvariant(FName::NameToDisplayString(TagPair.Key.ToString(), /*bIsBool*/false));
 				}
+			}
 
-				// Add a mapping from the sanitized display name back to the internal tag name
-				// This is useful as people only see the display name in the UI, so are more likely to use it
+			// Add a mapping from the sanitized display name back to the internal tag name
+			// This is useful as people only see the display name in the UI, so are more likely to use it
+			{
+				const FName SanitizedDisplayName = MakeObjectNameFromDisplayLabel(TagCache.DisplayName.ToString(), NAME_None);
+				if (TagPair.Key != SanitizedDisplayName)
 				{
-					const FName SanitizedDisplayName = MakeObjectNameFromDisplayLabel(TagCache.DisplayName.ToString(), NAME_None);
-					if (TagPair.Key != SanitizedDisplayName)
-					{
-						ClassCache->DisplayNameToTagNameMap.Add(SanitizedDisplayName, TagPair.Key);
-					}
+					ClassCache->DisplayNameToTagNameMap.Add(SanitizedDisplayName, TagPair.Key);
 				}
 			}
 		}
 
 		FWriteScopeLock WriteScope(Lock);
-		ClassToCacheMap.Add(InClassName, ClassCache);
+		ClassToCacheMap.Add(ActualClassPath, ClassCache);
+		if (InClassName != ActualClassPath)
+		{
+			ClassToCacheMap.Add(InClassName, ClassCache);
+		}
 	}
-
-	check(ClassCache);
-	return *ClassCache;
 }
diff --git a/Engine/Plugins/Editor/ContentBrowser/ContentBrowserAssetDataSource/Source/ContentBrowserAssetDataSource/Private/ContentBrowserAssetDataSource.cpp b/Engine/Plugins/Editor/ContentBrowser/ContentBrowserAssetDataSource/Source/ContentBrowserAssetDataSource/Private/ContentBrowserAssetDataSource.cpp
--- a/Engine/Plugins/Editor/ContentBrowser/ContentBrowserAssetDataSource/Source/ContentBrowserAssetDataSource/Private/ContentBrowserAssetDataSource.cpp
+++ b/Engine/Plugins/Editor/ContentBrowser/ContentBrowserAssetDataSource/Source/ContentBrowserAssetDataSource/Private/ContentBrowserAssetDataSource.cpp
@@ -257,7 +257,10 @@ void UContentBrowserAssetDataSource::Initialize(const bool InAutoRegister)
 	});
 	AssetRegistry->EnumerateAllAssets([this, &PropertyTagCache ](const FAssetData& InAssetData)
 		{
-			PropertyTagCache.GetCacheForClass(InAssetData.AssetClassPath);
+			if (InAssetData.GetOptionalOuterPathName().IsNone())
+			{
+				PropertyTagCache.TryCacheClass(InAssetData.AssetClassPath);
+			}
 			OnPathPopulated(InAssetData);
 			return true;
 		}, /*bIncludeOnlyOnDiskAssets*/true);
@@ -2125,12 +2128,12 @@ void UContentBrowserAssetDataSource::EnumerateItemsMatchingFilter(const FContent
 			return MoveTemp(Assets);
 		});
 
-		TArray<FAssetData> InMemoryAssets;
-		AssetRegistry->GetInMemoryAssets(AssetDataFilter->InclusiveFilter, InMemoryAssets);
+			TArray<FAssetData> InMemoryAssets;
+			AssetRegistry->GetInMemoryAssets(AssetDataFilter->InclusiveFilter, InMemoryAssets);
 
 		TSet<FName> IgnorePackages;
-		ProduceAssets(InMemoryAssets, IgnorePackages);
-		Algo::Transform(InMemoryAssets, IgnorePackages, [](const FAssetData& AssetData) { return AssetData.PackageName; });
+			ProduceAssets(InMemoryAssets, IgnorePackages);
+			Algo::Transform(InMemoryAssets, IgnorePackages, [](const FAssetData& AssetData) { return AssetData.PackageName; });
 
 		DiskTask.BusyWait();
 		TArray<FAssetData> DiskAssets = MoveTemp(DiskTask.GetResult());
@@ -3146,7 +3149,10 @@ void UContentBrowserAssetDataSource::OnAssetsAdded(TConstArrayView<FAssetData> I
 	FAssetPropertyTagCache& Cache = FAssetPropertyTagCache::Get();
 	for (const FAssetData& InAssetData : InAssets)
 	{
-		Cache.GetCacheForClass(InAssetData.AssetClassPath);
+		if (InAssetData.GetOptionalOuterPathName().IsNone())
+		{
+			Cache.TryCacheClass(InAssetData.AssetClassPath);
+		}
 
 		if (ContentBrowserAssetData::IsPrimaryAsset(InAssetData))
 		{
diff --git a/Engine/Plugins/Editor/ContentBrowser/ContentBrowserAssetDataSource/Source/ContentBrowserAssetDataSource/Public/AssetPropertyTagCache.h b/Engine/Plugins/Editor/ContentBrowser/ContentBrowserAssetDataSource/Source/ContentBrowserAssetDataSource/Public/AssetPropertyTagCache.h
--- a/Engine/Plugins/Editor/ContentBrowser/ContentBrowserAssetDataSource/Source/ContentBrowserAssetDataSource/Public/AssetPropertyTagCache.h
+++ b/Engine/Plugins/Editor/ContentBrowser/ContentBrowserAssetDataSource/Source/ContentBrowserAssetDataSource/Public/AssetPropertyTagCache.h
@@ -58,6 +58,9 @@ public:
 	/** Get the singleton instance */
 	static FAssetPropertyTagCache& Get();
 
+	/** Try and populate the cache for the given class if it is loaded */
+	void TryCacheClass(FTopLevelAssetPath InClassName);
+
 	/** Get (or populate) the cache for the given asset class */
 	const FClassPropertyTagCache& GetCacheForClass(FTopLevelAssetPath InClassName);
 
