Populate class property tag cache when asset data source gets 'modified' notifications, because some importers only 'modify' some of their assets and do not 'add' them. (e.g. when importing multiple assets).

Do not trigger ensure on missing property tag cache for non-native classes, which may not be loaded.

diff --git a/Engine/Plugins/Editor/ContentBrowser/ContentBrowserAssetDataSource/Source/ContentBrowserAssetDataSource/Private/AssetPropertyTagCache.cpp b/Engine/Plugins/Editor/ContentBrowser/ContentBrowserAssetDataSource/Source/ContentBrowserAssetDataSource/Private/AssetPropertyTagCache.cpp
--- a/Engine/Plugins/Editor/ContentBrowser/ContentBrowserAssetDataSource/Source/ContentBrowserAssetDataSource/Private/AssetPropertyTagCache.cpp
+++ b/Engine/Plugins/Editor/ContentBrowser/ContentBrowserAssetDataSource/Source/ContentBrowserAssetDataSource/Private/AssetPropertyTagCache.cpp
@@ -139,4 +139,8 @@ void FAssetPropertyTagCache::TryCacheClass(FTopLevelAssetPath InClassName)
 			ClassToCacheMap.Add(InClassName, ClassCache);
 		}
 	}
+	else 
+	{
+		// Do not populate the cache with something empty and do not report any errors as assets may have non-loaded types e.g. blueprinted data assets
+	}
 }
diff --git a/Engine/Plugins/Editor/ContentBrowser/ContentBrowserAssetDataSource/Source/ContentBrowserAssetDataSource/Private/ContentBrowserAssetDataCore.cpp b/Engine/Plugins/Editor/ContentBrowser/ContentBrowserAssetDataSource/Source/ContentBrowserAssetDataSource/Private/ContentBrowserAssetDataCore.cpp
--- a/Engine/Plugins/Editor/ContentBrowser/ContentBrowserAssetDataSource/Source/ContentBrowserAssetDataSource/Private/ContentBrowserAssetDataCore.cpp
+++ b/Engine/Plugins/Editor/ContentBrowser/ContentBrowserAssetDataSource/Source/ContentBrowserAssetDataSource/Private/ContentBrowserAssetDataCore.cpp
@@ -36,6 +36,13 @@ DEFINE_LOG_CATEGORY(LogContentBrowserAssetDataSource);
 namespace ContentBrowserAssetData
 {
 
+// Used to allow missing property metadata for non-native classes, as we're more lenient about those classes not being loaded & their data cached
+bool IsNonNativeClass(FTopLevelAssetPath ClassPath)
+{
+	FNameBuilder PackageName(ClassPath.GetPackageName());
+	return !FPackageName::IsScriptPackage(PackageName.ToView());	
+}
+
 FContentBrowserItemData CreateAssetFolderItem(
 	UContentBrowserDataSource* InOwnerDataSource,
 	const FName InVirtualPath,
@@ -1792,7 +1799,8 @@ bool GetAssetDataAttribute(const FAssetData& InAssetData, const bool InIncludeMe
 		FAssetDataTagMapSharedView::FFindTagResult FoundValue = InAssetData.TagsAndValues.FindTag(FoundAttributeKey);
 		if (!FoundValue.IsSet())
 		{
-			if (ensureMsgf(ClassPropertyTagCache, TEXT("FAssetPropertyTagCache not populated for type %s"), *WriteToString<256>(InAssetData.AssetClassPath)))
+			ensureMsgf(ClassPropertyTagCache || IsNonNativeClass(InAssetData.AssetClassPath), TEXT("FAssetPropertyTagCache not populated for type %s when looking for attribute %s"), *WriteToString<256>(InAssetData.AssetClassPath), *WriteToString<256>(InAttributeKey));
+			if (ClassPropertyTagCache)
 			{
 				// Check to see if the key we were given resolves as an alias
 				FoundAttributeKey = ClassPropertyTagCache->GetTagNameFromAlias(FoundAttributeKey);
@@ -1890,6 +1898,9 @@ bool GetAssetDataAttributes(const FAssetData& InAssetData, const bool InIncludeM
 			}
 		}
 
+		ensureMsgf(!InIncludeMetaData || ClassPropertyTagCache || IsNonNativeClass(InAssetData.AssetClassPath), TEXT("FAssetPropertyTagCache not populated for type %s when fetching all attributes"), 
+			*WriteToString<256>(InAssetData.AssetClassPath));
+
 		OutAttributeValues.Reserve(OutAttributeValues.Num() + InAssetData.TagsAndValues.Num());
 		for (const auto& TagAndValue : InAssetData.TagsAndValues)
 		{
diff --git a/Engine/Plugins/Editor/ContentBrowser/ContentBrowserAssetDataSource/Source/ContentBrowserAssetDataSource/Private/ContentBrowserAssetDataSource.cpp b/Engine/Plugins/Editor/ContentBrowser/ContentBrowserAssetDataSource/Source/ContentBrowserAssetDataSource/Private/ContentBrowserAssetDataSource.cpp
--- a/Engine/Plugins/Editor/ContentBrowser/ContentBrowserAssetDataSource/Source/ContentBrowserAssetDataSource/Private/ContentBrowserAssetDataSource.cpp
+++ b/Engine/Plugins/Editor/ContentBrowser/ContentBrowserAssetDataSource/Source/ContentBrowserAssetDataSource/Private/ContentBrowserAssetDataSource.cpp
@@ -3190,6 +3190,7 @@ void UContentBrowserAssetDataSource::OnAssetUpdated(const FAssetData& InAssetDat
 {
 	if (ContentBrowserAssetData::IsPrimaryAsset(InAssetData))
 	{
+		FAssetPropertyTagCache::Get().TryCacheClass(InAssetData.AssetClassPath);
 		QueueItemDataUpdate(FContentBrowserItemDataUpdate::MakeItemModifiedUpdate(CreateAssetFileItem(InAssetData)));
 	}
 }
@@ -3198,6 +3199,7 @@ void UContentBrowserAssetDataSource::OnAssetUpdatedOnDisk(const FAssetData& InAs
 {
 	if (ContentBrowserAssetData::IsPrimaryAsset(InAssetData))
 	{
+		FAssetPropertyTagCache::Get().TryCacheClass(InAssetData.AssetClassPath);
 		QueueItemDataUpdate(FContentBrowserItemDataUpdate::MakeItemModifiedUpdate(CreateAssetFileItem(InAssetData)));
 	}
 }
@@ -3207,6 +3209,7 @@ void UContentBrowserAssetDataSource::OnObjectPropertyChanged(UObject* InObject,
 	if (InObject && InObject->IsAsset() && ContentBrowserAssetData::IsPrimaryAsset(InObject))
 	{
 		FAssetData AssetData(InObject);
+		FAssetPropertyTagCache::Get().TryCacheClass(AssetData.AssetClassPath);
 		QueueItemDataUpdate(FContentBrowserItemDataUpdate::MakeItemModifiedUpdate(CreateAssetFileItem(AssetData)));
 	}
 }
@@ -3216,6 +3219,7 @@ void UContentBrowserAssetDataSource::OnObjectPreSave(UObject* InObject, FObjectP
 	if (InObject && InObject->IsAsset() && ContentBrowserAssetData::IsPrimaryAsset(InObject))
 	{
 		FAssetData AssetData(InObject);
+		FAssetPropertyTagCache::Get().TryCacheClass(AssetData.AssetClassPath);
 		QueueItemDataUpdate(FContentBrowserItemDataUpdate::MakeItemModifiedUpdate(CreateAssetFileItem(AssetData)));
 	}
 }
