Prepare asset property tag cache for in-memory assets which have not been scanned on disk by the asset registry yet.

diff --git a/Engine/Plugins/Editor/ContentBrowser/ContentBrowserAssetDataSource/Source/ContentBrowserAssetDataSource/Private/ContentBrowserAssetDataSource.cpp b/Engine/Plugins/Editor/ContentBrowser/ContentBrowserAssetDataSource/Source/ContentBrowserAssetDataSource/Private/ContentBrowserAssetDataSource.cpp
--- a/Engine/Plugins/Editor/ContentBrowser/ContentBrowserAssetDataSource/Source/ContentBrowserAssetDataSource/Private/ContentBrowserAssetDataSource.cpp
+++ b/Engine/Plugins/Editor/ContentBrowser/ContentBrowserAssetDataSource/Source/ContentBrowserAssetDataSource/Private/ContentBrowserAssetDataSource.cpp
@@ -2128,12 +2128,19 @@ void UContentBrowserAssetDataSource::EnumerateItemsMatchingFilter(const FContent
 			return MoveTemp(Assets);
 		});
 
-			TArray<FAssetData> InMemoryAssets;
-			AssetRegistry->GetInMemoryAssets(AssetDataFilter->InclusiveFilter, InMemoryAssets);
+		TArray<FAssetData> InMemoryAssets;
+		AssetRegistry->GetInMemoryAssets(AssetDataFilter->InclusiveFilter, InMemoryAssets);
+
+		// Prepare asset tag info for loaded assets which may not yet have been scanned by the asset registry
+		FAssetPropertyTagCache& TagCache = FAssetPropertyTagCache::Get();
+		for (const FAssetData& InMemoryAsset : InMemoryAssets)
+		{
+			TagCache.TryCacheClass(InMemoryAsset.AssetClassPath);
+		}
 
 		TSet<FName> IgnorePackages;
-			ProduceAssets(InMemoryAssets, IgnorePackages);
-			Algo::Transform(InMemoryAssets, IgnorePackages, [](const FAssetData& AssetData) { return AssetData.PackageName; });
+		ProduceAssets(InMemoryAssets, IgnorePackages);
+		Algo::Transform(InMemoryAssets, IgnorePackages, [](const FAssetData& AssetData) { return AssetData.PackageName; });
 
 		DiskTask.BusyWait();
 		TArray<FAssetData> DiskAssets = MoveTemp(DiskTask.GetResult());
@@ -3149,6 +3156,8 @@ void UContentBrowserAssetDataSource::OnAssetsAdded(TConstArrayView<FAssetData> I
 	FAssetPropertyTagCache& Cache = FAssetPropertyTagCache::Get();
 	for (const FAssetData& InAssetData : InAssets)
 	{
+		UE_LOG(LogContentBrowserAssetDataSource, VeryVerbose, TEXT("OnAssetsAdded: %s"), *WriteToString<256>(InAssetData.GetSoftObjectPath()));
+
 		if (InAssetData.GetOptionalOuterPathName().IsNone())
 		{
 			Cache.TryCacheClass(InAssetData.AssetClassPath);
@@ -3168,6 +3177,7 @@ void UContentBrowserAssetDataSource::OnAssetRemoved(const FAssetData& InAssetDat
 {
 	if (ContentBrowserAssetData::IsPrimaryAsset(InAssetData))
 	{
+		UE_LOG(LogContentBrowserAssetDataSource, VeryVerbose, TEXT("OnAssetRemoved: %s"), *WriteToString<256>(InAssetData.GetSoftObjectPath()));
 		QueueItemDataUpdate(FContentBrowserItemDataUpdate::MakeItemRemovedUpdate(CreateAssetFileItem(InAssetData)));
 	}
 }
@@ -3176,6 +3186,8 @@ void UContentBrowserAssetDataSource::OnAssetRenamed(const FAssetData& InAssetDat
 {
 	if (ContentBrowserAssetData::IsPrimaryAsset(InAssetData))
 	{
+		UE_LOG(LogContentBrowserAssetDataSource, VeryVerbose, TEXT("OnAssetRenamed: %s"), *WriteToString<256>(InAssetData.GetSoftObjectPath()));
+
 		// The owner folder of this asset is no longer considered empty
 		OnPathPopulated(InAssetData);
 
@@ -3190,6 +3202,8 @@ void UContentBrowserAssetDataSource::OnAssetUpdated(const FAssetData& InAssetDat
 {
 	if (ContentBrowserAssetData::IsPrimaryAsset(InAssetData))
 	{
+		UE_LOG(LogContentBrowserAssetDataSource, VeryVerbose, TEXT("OnAssetUpdated: %s"), *WriteToString<256>(InAssetData.GetSoftObjectPath()));
+
 		FAssetPropertyTagCache::Get().TryCacheClass(InAssetData.AssetClassPath);
 		QueueItemDataUpdate(FContentBrowserItemDataUpdate::MakeItemModifiedUpdate(CreateAssetFileItem(InAssetData)));
 	}
@@ -3199,6 +3213,8 @@ void UContentBrowserAssetDataSource::OnAssetUpdatedOnDisk(const FAssetData& InAs
 {
 	if (ContentBrowserAssetData::IsPrimaryAsset(InAssetData))
 	{
+		UE_LOG(LogContentBrowserAssetDataSource, VeryVerbose, TEXT("OnAssetUpdatedOnDisk: %s"), *WriteToString<256>(InAssetData.GetSoftObjectPath()));
+
 		FAssetPropertyTagCache::Get().TryCacheClass(InAssetData.AssetClassPath);
 		QueueItemDataUpdate(FContentBrowserItemDataUpdate::MakeItemModifiedUpdate(CreateAssetFileItem(InAssetData)));
 	}
diff --git a/Engine/Plugins/Editor/ContentBrowser/ContentBrowserAssetDataSource/Source/ContentBrowserAssetDataSource/Public/ContentBrowserAssetDataCore.h b/Engine/Plugins/Editor/ContentBrowser/ContentBrowserAssetDataSource/Source/ContentBrowserAssetDataSource/Public/ContentBrowserAssetDataCore.h
--- a/Engine/Plugins/Editor/ContentBrowser/ContentBrowserAssetDataSource/Source/ContentBrowserAssetDataSource/Public/ContentBrowserAssetDataCore.h
+++ b/Engine/Plugins/Editor/ContentBrowser/ContentBrowserAssetDataSource/Source/ContentBrowserAssetDataSource/Public/ContentBrowserAssetDataCore.h
@@ -17,7 +17,7 @@ class UObject;
 class UToolMenu;
 class UContentBrowserDataSource;
 
-DECLARE_LOG_CATEGORY_EXTERN(LogContentBrowserAssetDataSource, Log, All);
+DECLARE_LOG_CATEGORY_EXTERN(LogContentBrowserAssetDataSource, Log, Log);
 
 namespace ContentBrowserAssetData
 {
