FAssetPropertyTagCache: Queue up classes which are not yet initialized when asset data is scanned. Try again after modules are initialized.
Loosen ensure about unavailable data to only fire if the class exists.

diff --git a/Engine/Plugins/Editor/ContentBrowser/ContentBrowserAssetDataSource/Source/ContentBrowserAssetDataSource/Private/AssetPropertyTagCache.cpp b/Engine/Plugins/Editor/ContentBrowser/ContentBrowserAssetDataSource/Source/ContentBrowserAssetDataSource/Private/AssetPropertyTagCache.cpp
--- a/Engine/Plugins/Editor/ContentBrowser/ContentBrowserAssetDataSource/Source/ContentBrowserAssetDataSource/Private/AssetPropertyTagCache.cpp
+++ b/Engine/Plugins/Editor/ContentBrowser/ContentBrowserAssetDataSource/Source/ContentBrowserAssetDataSource/Private/AssetPropertyTagCache.cpp
@@ -39,6 +39,20 @@ const FAssetPropertyTagCache::FClassPropertyTagCache& FAssetPropertyTagCache::Ge
 	return *ClassCache;
 }
 
+void FAssetPropertyTagCache::CachePendingClasses()
+{
+	TSet<FTopLevelAssetPath> PendingCopy;
+	{ 
+		FWriteScopeLock WriteScope(Lock);
+		PendingCopy = PendingClasses;
+	}
+
+	for (FTopLevelAssetPath Path : PendingCopy)
+	{
+		TryCacheClass(Path);
+	}
+}
+
 void FAssetPropertyTagCache::TryCacheClass(FTopLevelAssetPath InClassName)
 {
 	TSharedPtr<FClassPropertyTagCache> ClassCache;
@@ -52,9 +66,6 @@ void FAssetPropertyTagCache::TryCacheClass(FTopLevelAssetPath InClassName)
 	{
 		return;
 	}
-	
-	// On a miss, race with other potential creations rather than hold the lock against queries for another class
-	UE_LOG(LogContentBrowserAssetDataSource, Verbose, TEXT("Creating asset property tag cache for %s"), *WriteToString<256>(InClassName));
 
 	auto GetAssetClass = [&InClassName]()
 	{
@@ -75,6 +86,9 @@ void FAssetPropertyTagCache::TryCacheClass(FTopLevelAssetPath InClassName)
 
 	if (UClass* AssetClass = GetAssetClass())
 	{
+		// On a miss, race with other potential creations rather than hold the lock against queries for another class
+		UE_LOG(LogContentBrowserAssetDataSource, Verbose, TEXT("Creating asset property tag cache for %s"), *WriteToString<256>(InClassName));
+
 		ClassCache = MakeShared<FClassPropertyTagCache>();
 		FTopLevelAssetPath ActualClassPath(AssetClass);
 
@@ -134,13 +148,17 @@ void FAssetPropertyTagCache::TryCacheClass(FTopLevelAssetPath InClassName)
 
 		FWriteScopeLock WriteScope(Lock);
 		ClassToCacheMap.Add(ActualClassPath, ClassCache);
+		PendingClasses.Remove(ActualClassPath);
 		if (InClassName != ActualClassPath)
 		{
 			ClassToCacheMap.Add(InClassName, ClassCache);
+			PendingClasses.Remove(InClassName);
 		}
 	}
 	else 
 	{
-		// Do not populate the cache with something empty and do not report any errors as assets may have non-loaded types e.g. blueprinted data assets
+		// Some assets may be scanned before their modules are initialized
+		FWriteScopeLock WriteScope(Lock);
+		PendingClasses.Add(InClassName);
 	}
 }
diff --git a/Engine/Plugins/Editor/ContentBrowser/ContentBrowserAssetDataSource/Source/ContentBrowserAssetDataSource/Private/ContentBrowserAssetDataCore.cpp b/Engine/Plugins/Editor/ContentBrowser/ContentBrowserAssetDataSource/Source/ContentBrowserAssetDataSource/Private/ContentBrowserAssetDataCore.cpp
--- a/Engine/Plugins/Editor/ContentBrowser/ContentBrowserAssetDataSource/Source/ContentBrowserAssetDataSource/Private/ContentBrowserAssetDataCore.cpp
+++ b/Engine/Plugins/Editor/ContentBrowser/ContentBrowserAssetDataSource/Source/ContentBrowserAssetDataSource/Private/ContentBrowserAssetDataCore.cpp
@@ -36,13 +36,6 @@ DEFINE_LOG_CATEGORY(LogContentBrowserAssetDataSource);
 namespace ContentBrowserAssetData
 {
 
-// Used to allow missing property metadata for non-native classes, as we're more lenient about those classes not being loaded & their data cached
-bool IsNonNativeClass(FTopLevelAssetPath ClassPath)
-{
-	FNameBuilder PackageName(ClassPath.GetPackageName());
-	return !FPackageName::IsScriptPackage(PackageName.ToView());	
-}
-
 FContentBrowserItemData CreateAssetFolderItem(
 	UContentBrowserDataSource* InOwnerDataSource,
 	const FName InVirtualPath,
@@ -1799,7 +1792,7 @@ bool GetAssetDataAttribute(const FAssetData& InAssetData, const bool InIncludeMe
 		FAssetDataTagMapSharedView::FFindTagResult FoundValue = InAssetData.TagsAndValues.FindTag(FoundAttributeKey);
 		if (!FoundValue.IsSet())
 		{
-			ensureMsgf(ClassPropertyTagCache || IsNonNativeClass(InAssetData.AssetClassPath), TEXT("FAssetPropertyTagCache not populated for type %s when looking for attribute %s"), *WriteToString<256>(InAssetData.AssetClassPath), *WriteToString<256>(InAttributeKey));
+			ensureMsgf(ClassPropertyTagCache || (FindObject<UClass>(InAssetData.AssetClassPath) == nullptr), TEXT("FAssetPropertyTagCache not populated for type %s when looking for attribute %s"), *WriteToString<256>(InAssetData.AssetClassPath), *WriteToString<256>(InAttributeKey));
 			if (ClassPropertyTagCache)
 			{
 				// Check to see if the key we were given resolves as an alias
@@ -1898,7 +1891,7 @@ bool GetAssetDataAttributes(const FAssetData& InAssetData, const bool InIncludeM
 			}
 		}
 
-		ensureMsgf(!InIncludeMetaData || ClassPropertyTagCache || IsNonNativeClass(InAssetData.AssetClassPath), TEXT("FAssetPropertyTagCache not populated for type %s when fetching all attributes"), 
+		ensureMsgf(!InIncludeMetaData || ClassPropertyTagCache || (FindObject<UClass>(InAssetData.AssetClassPath) == nullptr), TEXT("FAssetPropertyTagCache not populated for type %s when fetching all attributes"), 
 			*WriteToString<256>(InAssetData.AssetClassPath));
 
 		OutAttributeValues.Reserve(OutAttributeValues.Num() + InAssetData.TagsAndValues.Num());
diff --git a/Engine/Plugins/Editor/ContentBrowser/ContentBrowserAssetDataSource/Source/ContentBrowserAssetDataSource/Private/ContentBrowserAssetDataSource.cpp b/Engine/Plugins/Editor/ContentBrowser/ContentBrowserAssetDataSource/Source/ContentBrowserAssetDataSource/Private/ContentBrowserAssetDataSource.cpp
--- a/Engine/Plugins/Editor/ContentBrowser/ContentBrowserAssetDataSource/Source/ContentBrowserAssetDataSource/Private/ContentBrowserAssetDataSource.cpp
+++ b/Engine/Plugins/Editor/ContentBrowser/ContentBrowserAssetDataSource/Source/ContentBrowserAssetDataSource/Private/ContentBrowserAssetDataSource.cpp
@@ -170,6 +170,19 @@ void UContentBrowserAssetDataSource::Initialize(const bool InAutoRegister)
 	// Listen for when assets are saved, listerns are notified in time despite presave because we queue updates for later processing 
 	FCoreUObjectDelegates::OnObjectPreSave.AddUObject(this, &UContentBrowserAssetDataSource::OnObjectPreSave);
 
+	// Listen for module initialization to update FAssetPropertyTagCache
+	FCoreUObjectDelegates::CompiledInUObjectsRegisteredDelegate.AddWeakLambda(this, [](FName, ECompiledInUObjectsRegisteredStatus){
+		FAssetPropertyTagCache::Get().CachePendingClasses();
+	});
+
+	// Listen for classes being loaded 
+	FCoreUObjectDelegates::OnAssetLoaded.AddWeakLambda(this, [](UObject* Object){
+		if (UClass* Class = Cast<UClass>(Object))
+		{
+			FAssetPropertyTagCache::Get().TryCacheClass(FTopLevelAssetPath(Class));
+		}
+	});
+
 	// Listen for new mount roots
 	FPackageName::OnContentPathMounted().AddUObject(this, &UContentBrowserAssetDataSource::OnContentPathMounted);
 	FPackageName::OnContentPathDismounted().AddUObject(this, &UContentBrowserAssetDataSource::OnContentPathDismounted);
diff --git a/Engine/Plugins/Editor/ContentBrowser/ContentBrowserAssetDataSource/Source/ContentBrowserAssetDataSource/Public/AssetPropertyTagCache.h b/Engine/Plugins/Editor/ContentBrowser/ContentBrowserAssetDataSource/Source/ContentBrowserAssetDataSource/Public/AssetPropertyTagCache.h
--- a/Engine/Plugins/Editor/ContentBrowser/ContentBrowserAssetDataSource/Source/ContentBrowserAssetDataSource/Public/AssetPropertyTagCache.h
+++ b/Engine/Plugins/Editor/ContentBrowser/ContentBrowserAssetDataSource/Source/ContentBrowserAssetDataSource/Public/AssetPropertyTagCache.h
@@ -61,6 +61,9 @@ public:
 	/** Try and populate the cache for the given class if it is loaded */
 	void TryCacheClass(FTopLevelAssetPath InClassName);
 
+	/** Try and populate the cache with classes which were not reigstered when TryCacheClass was called. */
+	void CachePendingClasses();
+
 	/** Get (or populate) the cache for the given asset class */
 	const FClassPropertyTagCache& GetCacheForClass(FTopLevelAssetPath InClassName);
 
@@ -75,6 +78,9 @@ private:
 	 * Values are not modified after construction, so they can safely be returned after releasing the lock. 
 	 */
 	TMap<FTopLevelAssetPath, TSharedPtr<FClassPropertyTagCache>> ClassToCacheMap;
+
+	/** Classes we'd like to register but were not loaded when an asset of that type was last scanned */
+	TSet<FTopLevelAssetPath> PendingClasses;
 };
 
 #if UE_ENABLE_INCLUDE_ORDER_DEPRECATED_IN_5_2
