Fixed the behavior of the Show Redirector filter when the class permission system is used.

The ExclusiveFilter act as filter to remove some items that can be showed via the inclusive filter, but the presence of a authorized type in the inclusive filter doesn't mean that something explicitly asked for that type to be shown.

[RN] [Content Browser] Fixed a bug that would make the redirectors always visible.

diff --git a/Engine/Plugins/Editor/ContentBrowser/ContentBrowserAssetDataSource/Source/ContentBrowserAssetDataSource/Private/ContentBrowserAssetDataSource.cpp b/Engine/Plugins/Editor/ContentBrowser/ContentBrowserAssetDataSource/Source/ContentBrowserAssetDataSource/Private/ContentBrowserAssetDataSource.cpp
--- a/Engine/Plugins/Editor/ContentBrowser/ContentBrowserAssetDataSource/Source/ContentBrowserAssetDataSource/Private/ContentBrowserAssetDataSource.cpp
+++ b/Engine/Plugins/Editor/ContentBrowser/ContentBrowserAssetDataSource/Source/ContentBrowserAssetDataSource/Private/ContentBrowserAssetDataSource.cpp
@@ -919,12 +919,7 @@ PRAGMA_ENABLE_DEPRECATION_WARNINGS
 
 			if (!Params.bIncludeRedirectors)
 			{
-				// Add redirectors to the class paths to exclude unless it's already in the class paths to explicitly include
-				FTopLevelAssetPath RedirectorClassPath(UObjectRedirector::StaticClass());
-				if (!CompiledInclusiveFilter.ClassPaths.Contains(RedirectorClassPath))
-				{
-					ExclusiveFilter.ClassPaths.Add(RedirectorClassPath);
-				}
+				ExclusiveFilter.ClassPaths.Add(FTopLevelAssetPath(UObjectRedirector::StaticClass()));
 			}
 
 			Params.AssetRegistry->CompileFilter(ExclusiveFilter, CompiledExclusiveFilter);
