Populate asset property tag cache when creating new assets with placeholder asset data.

diff --git a/Engine/Plugins/Editor/ContentBrowser/ContentBrowserAssetDataSource/Source/ContentBrowserAssetDataSource/Private/ContentBrowserAssetDataSource.cpp b/Engine/Plugins/Editor/ContentBrowser/ContentBrowserAssetDataSource/Source/ContentBrowserAssetDataSource/Private/ContentBrowserAssetDataSource.cpp
--- a/Engine/Plugins/Editor/ContentBrowser/ContentBrowserAssetDataSource/Source/ContentBrowserAssetDataSource/Private/ContentBrowserAssetDataSource.cpp
+++ b/Engine/Plugins/Editor/ContentBrowser/ContentBrowserAssetDataSource/Source/ContentBrowserAssetDataSource/Private/ContentBrowserAssetDataSource.cpp
@@ -3663,6 +3663,8 @@ void UContentBrowserAssetDataSource::OnBeginCreateAsset(const FName InDefaultAss
 		return;
 	}
 
+	FAssetPropertyTagCache::Get().TryCacheClass(FTopLevelAssetPath(ClassToUse));
+
 	FContentBrowserModule& ContentBrowserModule = FModuleManager::Get().GetModuleChecked<FContentBrowserModule>(TEXT("ContentBrowser"));
 	FContentBrowserItemPath AssetPathToUse = ContentBrowserModule.Get().GetInitialPathToSaveAsset(FContentBrowserItemPath(InPackagePath, EContentBrowserPathType::Internal));
 
