Explicitly fetch asset registry tags for only in-memory/modified assets in asset data source.

diff --git a/Engine/Plugins/Editor/ContentBrowser/ContentBrowserAssetDataSource/Source/ContentBrowserAssetDataSource/Private/ContentBrowserAssetDataSource.cpp b/Engine/Plugins/Editor/ContentBrowser/ContentBrowserAssetDataSource/Source/ContentBrowserAssetDataSource/Private/ContentBrowserAssetDataSource.cpp
--- a/Engine/Plugins/Editor/ContentBrowser/ContentBrowserAssetDataSource/Source/ContentBrowserAssetDataSource/Private/ContentBrowserAssetDataSource.cpp
+++ b/Engine/Plugins/Editor/ContentBrowser/ContentBrowserAssetDataSource/Source/ContentBrowserAssetDataSource/Private/ContentBrowserAssetDataSource.cpp
@@ -56,6 +56,14 @@ namespace AssetDataSource
 		TEXT("Set to 0 to disable internal parallelism inside data source in case of threading issues."),
 		ECVF_Default
 	);
+
+	bool bOptimizeEnumerateInMemoryAssets = true;
+	FAutoConsoleVariableRef CVarOptimizeEnumerateInMemoryAssets(
+		TEXT("AssetDataSource.OptimizeEnumerateInMemoryAssets"),
+		bOptimizeEnumerateInMemoryAssets ,
+		TEXT("1: Explicitly fetch fresh asset data for only new/dirty assets. 0: Fetch fresh asset data for all loaded assets."),
+		ECVF_Default
+	);
 }
 
 enum class EContentBrowserFolderAttributes : uint8
@@ -2005,6 +2013,7 @@ void UContentBrowserAssetDataSource::EnumerateItemsMatchingFilter(const FContent
 	EnumerateItemsMatchingFilter(InFilter, TGetOrEnumerateSink<FContentBrowserItemData>(MoveTemp(InCallback)));
 }
 
+
 void UContentBrowserAssetDataSource::EnumerateItemsMatchingFilter(const FContentBrowserDataCompiledFilter& InFilter, const TGetOrEnumerateSink<FContentBrowserItemData>& InSink)
 {
 	const FContentBrowserDataFilterList* FilterList = InFilter.CompiledFilters.Find(this);
@@ -2071,8 +2080,18 @@ void UContentBrowserAssetDataSource::EnumerateItemsMatchingFilter(const FContent
 			return;
 		}
 
-		auto ProduceAssets = [&AssetDataFilter, &InSink, this](TArray<FAssetData>& Assets, TSet<FName> IgnorePackageNames) {
+		auto ProduceAssets = [&AssetDataFilter, &InSink, this](TArray<FAssetData>& Assets, const TSet<FName>& IgnorePackageNames, bool bUpdatePropertyTagCache) {
 			InSink.ReserveMore(Assets.Num());
+			
+			if (bUpdatePropertyTagCache)
+			{
+				FAssetPropertyTagCache& TagCache = FAssetPropertyTagCache::Get();
+				for (const FAssetData& AssetData : Assets)
+				{
+					TagCache.TryCacheClass(AssetData.AssetClassPath);
+				}
+			}
+			
 			for (FAssetData& AssetData : Assets)
 			{
 				if (IgnorePackageNames.Contains(AssetData.PackageName))
@@ -2136,23 +2155,43 @@ void UContentBrowserAssetDataSource::EnumerateItemsMatchingFilter(const FContent
 			return MoveTemp(Assets);
 		});
 
-		TArray<FAssetData> InMemoryAssets;
-		AssetRegistry->GetInMemoryAssets(AssetDataFilter->InclusiveFilter, InMemoryAssets);
-
-		// Prepare asset tag info for loaded assets which may not yet have been scanned by the asset registry
-		FAssetPropertyTagCache& TagCache = FAssetPropertyTagCache::Get();
-		for (const FAssetData& InMemoryAsset : InMemoryAssets)
+		TSet<FName> IgnorePackages;
+		if (AssetDataFilter->InclusiveFilter.PackageNames.Num() != 0 || !AssetDataSource::bOptimizeEnumerateInMemoryAssets) 
 		{
-			TagCache.TryCacheClass(InMemoryAsset.AssetClassPath);
+			TArray<FAssetData> InMemoryAssets;
+			AssetRegistry->GetInMemoryAssets(AssetDataFilter->InclusiveFilter, InMemoryAssets);
+			ProduceAssets(InMemoryAssets, IgnorePackages, true);
+			Algo::Transform(InMemoryAssets, IgnorePackages, [](const FAssetData& AssetData) { return AssetData.PackageName; });
 		}
+		else
+		{
+			TArray<FAssetData> InMemoryAssets;
+			FARCompiledFilter InMemoryFilter = AssetDataFilter->InclusiveFilter;
+			ForEachObjectOfClass(UPackage::StaticClass(), [&InMemoryFilter](UObject* Object)
+			{
+				UPackage* Package = CastChecked<UPackage>(Object);
+				if (Package->HasAnyFlags(RF_ClassDefaultObject))
+				{
+					return;
+				}
 
-		TSet<FName> IgnorePackages;
-		ProduceAssets(InMemoryAssets, IgnorePackages);
-		Algo::Transform(InMemoryAssets, IgnorePackages, [](const FAssetData& AssetData) { return AssetData.PackageName; });
+				if (Package->IsDirty() || Package->HasAnyPackageFlags(PKG_NewlyCreated))
+				{
+					InMemoryFilter.PackageNames.Add(Package->GetFName());
+				}
+			});
+			if (InMemoryFilter.PackageNames.Num() > 0)
+			{
+				AssetRegistry->GetInMemoryAssets(InMemoryFilter, InMemoryAssets);
+
+				ProduceAssets(InMemoryAssets, IgnorePackages, true);
+				Algo::Transform(InMemoryAssets, IgnorePackages, [](const FAssetData& AssetData) { return AssetData.PackageName; });
+			}
+		}
 
 		DiskTask.BusyWait();
 		TArray<FAssetData> DiskAssets = MoveTemp(DiskTask.GetResult());
-		ProduceAssets(DiskAssets, IgnorePackages);
+		ProduceAssets(DiskAssets, IgnorePackages, false);
 	}
 }
 
