Slightly optimize population of localization context menu options by collecting and deduplicating pacakge roots before collecting cultures.

diff --git a/Engine/Plugins/Editor/ContentBrowser/ContentBrowserAssetDataSource/Source/ContentBrowserAssetDataSource/Private/AssetFileContextMenu.cpp b/Engine/Plugins/Editor/ContentBrowser/ContentBrowserAssetDataSource/Source/ContentBrowserAssetDataSource/Private/AssetFileContextMenu.cpp
--- a/Engine/Plugins/Editor/ContentBrowser/ContentBrowserAssetDataSource/Source/ContentBrowserAssetDataSource/Private/AssetFileContextMenu.cpp
+++ b/Engine/Plugins/Editor/ContentBrowser/ContentBrowserAssetDataSource/Source/ContentBrowserAssetDataSource/Private/AssetFileContextMenu.cpp
@@ -649,6 +649,7 @@ void FAssetFileContextMenu::MakeAssetLocalizationSubMenu(UToolMenu* Menu)
 		bool bIncludeEngineCultures = false;
 		bool bIncludeProjectCultures = false;
 
+		TSet<FString> Roots;
 		for (const FAssetData& Asset : SelectedAssets)
 		{
 			const FString AssetPath = Asset.GetObjectPathString();
@@ -662,18 +663,21 @@ void FAssetFileContextMenu::MakeAssetLocalizationSubMenu(UToolMenu* Menu)
 				bIncludeProjectCultures = true;
 			}
 
+			FString AssetLocalizationRoot;
+			if (FPackageLocalizationUtil::GetLocalizedRoot(AssetPath, FString(), AssetLocalizationRoot))
 			{
-				FString AssetLocalizationRoot;
-				if (FPackageLocalizationUtil::GetLocalizedRoot(AssetPath, FString(), AssetLocalizationRoot))
-				{
-					FString AssetLocalizationFileRoot;
-					if (FPackageName::TryConvertLongPackageNameToFilename(AssetLocalizationRoot, AssetLocalizationFileRoot))
-					{
-						TArray<FString> CulturePaths;
-						CulturePaths.Add(MoveTemp(AssetLocalizationFileRoot));
-						CultureNames.Append(TextLocalizationResourceUtil::GetLocalizedCultureNames(CulturePaths));
-					}
-				}
+				Roots.Add(AssetLocalizationRoot);	
+			}
+		}
+
+		for (const FString& AssetLocalizationRoot : Roots)
+		{
+			FString AssetLocalizationFileRoot;
+			if (FPackageName::TryConvertLongPackageNameToFilename(AssetLocalizationRoot, AssetLocalizationFileRoot))
+			{
+				TArray<FString> CulturePaths;
+				CulturePaths.Add(MoveTemp(AssetLocalizationFileRoot));
+				CultureNames.Append(TextLocalizationResourceUtil::GetLocalizedCultureNames(CulturePaths));
 			}
 		}
 
