Fixed an issue where "the select actor using this asset" tool would select all the blueprints actor instead of the one(s) that actually use the specified BP asset.

The dependencies graph was expended to much. Since we aslo scan the class data we don't need to add the direct link of the actor to the assets.

diff --git a/Engine/Plugins/Editor/ContentBrowser/ContentBrowserAssetDataSource/Source/ContentBrowserAssetDataSource/Private/AssetFileContextMenu.cpp b/Engine/Plugins/Editor/ContentBrowser/ContentBrowserAssetDataSource/Source/ContentBrowserAssetDataSource/Private/AssetFileContextMenu.cpp
--- a/Engine/Plugins/Editor/ContentBrowser/ContentBrowserAssetDataSource/Source/ContentBrowserAssetDataSource/Private/AssetFileContextMenu.cpp
+++ b/Engine/Plugins/Editor/ContentBrowser/ContentBrowserAssetDataSource/Source/ContentBrowserAssetDataSource/Private/AssetFileContextMenu.cpp
@@ -1412,19 +1412,6 @@ struct WorldReferenceGenerator : public FFindReferencedAssets
 				}
 			}
 		}
-
-		TArray<UObject*> ReferencedObjects;
-		// Special case for blueprints
-		for (AActor* Actor : FActorRange(World))
-		{
-			ReferencedObjects.Reset();
-			Actor->GetReferencedContentObjects(ReferencedObjects);
-			for(UObject* Reference : ReferencedObjects)
-			{
-				auto& Objects = ReferenceGraph.FindOrAdd(Reference);
-				Objects.Add(Actor);
-			}
-		}
 	}
 
 	void MarkAllObjects()
