Fixing an issue where the content browser would try to create a widget for an asset that the asset registry hadn't broadcast a callback about yet. The result was that data content browser relied on had not yet been produced.

#rnx

diff --git a/Engine/Plugins/Editor/ContentBrowser/ContentBrowserAssetDataSource/Source/ContentBrowserAssetDataSource/Private/ContentBrowserAssetDataSource.cpp b/Engine/Plugins/Editor/ContentBrowser/ContentBrowserAssetDataSource/Source/ContentBrowserAssetDataSource/Private/ContentBrowserAssetDataSource.cpp
--- a/Engine/Plugins/Editor/ContentBrowser/ContentBrowserAssetDataSource/Source/ContentBrowserAssetDataSource/Private/ContentBrowserAssetDataSource.cpp
+++ b/Engine/Plugins/Editor/ContentBrowser/ContentBrowserAssetDataSource/Source/ContentBrowserAssetDataSource/Private/ContentBrowserAssetDataSource.cpp
@@ -2080,16 +2080,13 @@ void UContentBrowserAssetDataSource::EnumerateItemsMatchingFilter(const FContent
 			return;
 		}
 
-		auto ProduceAssets = [&AssetDataFilter, &InSink, this](TArray<FAssetData>& Assets, const TSet<FName>& IgnorePackageNames, bool bUpdatePropertyTagCache) {
+		auto ProduceAssets = [&AssetDataFilter, &InSink, this](TArray<FAssetData>& Assets, const TSet<FName>& IgnorePackageNames) {
 			InSink.ReserveMore(Assets.Num());
 			
-			if (bUpdatePropertyTagCache)
+			FAssetPropertyTagCache& TagCache = FAssetPropertyTagCache::Get();
+			for (const FAssetData& AssetData : Assets)
 			{
-				FAssetPropertyTagCache& TagCache = FAssetPropertyTagCache::Get();
-				for (const FAssetData& AssetData : Assets)
-				{
-					TagCache.TryCacheClass(AssetData.AssetClassPath);
-				}
+				TagCache.TryCacheClass(AssetData.AssetClassPath);
 			}
 			
 			for (FAssetData& AssetData : Assets)
@@ -2160,7 +2157,7 @@ void UContentBrowserAssetDataSource::EnumerateItemsMatchingFilter(const FContent
 		{
 			TArray<FAssetData> InMemoryAssets;
 			AssetRegistry->GetInMemoryAssets(AssetDataFilter->InclusiveFilter, InMemoryAssets);
-			ProduceAssets(InMemoryAssets, IgnorePackages, true);
+			ProduceAssets(InMemoryAssets, IgnorePackages);
 			Algo::Transform(InMemoryAssets, IgnorePackages, [](const FAssetData& AssetData) { return AssetData.PackageName; });
 		}
 		else
@@ -2184,14 +2181,14 @@ void UContentBrowserAssetDataSource::EnumerateItemsMatchingFilter(const FContent
 			{
 				AssetRegistry->GetInMemoryAssets(InMemoryFilter, InMemoryAssets);
 
-				ProduceAssets(InMemoryAssets, IgnorePackages, true);
+				ProduceAssets(InMemoryAssets, IgnorePackages);
 				Algo::Transform(InMemoryAssets, IgnorePackages, [](const FAssetData& AssetData) { return AssetData.PackageName; });
 			}
 		}
 
 		DiskTask.BusyWait();
 		TArray<FAssetData> DiskAssets = MoveTemp(DiskTask.GetResult());
-		ProduceAssets(DiskAssets, IgnorePackages, false);
+		ProduceAssets(DiskAssets, IgnorePackages);
 	}
 }
 
