Removes unnecessary AssetData.GetPackage() call to fetch the Package path when it's available via AssetData.PackageName. Calling GetPackage() will load the package would can incur massive stalls when loading large packages (e.g. levels)

#rnx

diff --git a/Engine/Plugins/Editor/ContentBrowser/ContentBrowserAssetDataSource/Source/ContentBrowserAssetDataSource/Private/ContentBrowserAssetDataCore.cpp b/Engine/Plugins/Editor/ContentBrowser/ContentBrowserAssetDataSource/Source/ContentBrowserAssetDataSource/Private/ContentBrowserAssetDataCore.cpp
--- a/Engine/Plugins/Editor/ContentBrowser/ContentBrowserAssetDataSource/Source/ContentBrowserAssetDataSource/Private/ContentBrowserAssetDataCore.cpp
+++ b/Engine/Plugins/Editor/ContentBrowser/ContentBrowserAssetDataSource/Source/ContentBrowserAssetDataSource/Private/ContentBrowserAssetDataCore.cpp
@@ -1402,19 +1402,11 @@ void AppendAssetObjectPath(const FAssetData& AssetData, FString& InOutStr)
 
 void AppendAssetPackageName(const FAssetData& AssetData, FString& InOutStr)
 {
-	UPackage* Package = AssetData.GetPackage();
-	if (InOutStr.IsEmpty())
-	{
-		InOutStr = Package->GetPathName();
-	}
-	else
+	if (!InOutStr.IsEmpty())
 	{
-		if (InOutStr.Len() > 0)
-		{
-			InOutStr += LINE_TERMINATOR;
-		}
-		InOutStr += Package->GetPathName();
+		InOutStr += LINE_TERMINATOR;
 	}
+	InOutStr += AssetData.PackageName.ToString();
 }
 
 bool AppendAssetFolderItemReference(IAssetRegistry* InAssetRegistry, const FContentBrowserAssetFolderItemDataPayload& InFolderPayload, FString& InOutStr)
