- Updated content browser ::IsAssetFileItemDirty to check for asset definion flag (and external objects if true)

diff --git a/Engine/Plugins/Editor/ContentBrowser/ContentBrowserAssetDataSource/Source/ContentBrowserAssetDataSource/Private/ContentBrowserAssetDataCore.cpp b/Engine/Plugins/Editor/ContentBrowser/ContentBrowserAssetDataSource/Source/ContentBrowserAssetDataSource/Private/ContentBrowserAssetDataCore.cpp
--- a/Engine/Plugins/Editor/ContentBrowser/ContentBrowserAssetDataSource/Source/ContentBrowserAssetDataSource/Private/ContentBrowserAssetDataCore.cpp
+++ b/Engine/Plugins/Editor/ContentBrowser/ContentBrowserAssetDataSource/Source/ContentBrowserAssetDataSource/Private/ContentBrowserAssetDataCore.cpp
@@ -1288,8 +1288,26 @@ bool IsItemDirty(const UContentBrowserDataSource* InOwnerDataSource, const FCont
 
 bool IsAssetFileItemDirty(const FContentBrowserAssetFileItemDataPayload& InAssetPayload)
 {
-	UPackage* AssetPackage = InAssetPayload.GetPackage();
-	return AssetPackage && AssetPackage->IsDirty();
+	if (const UPackage* AssetPackage = InAssetPayload.GetPackage())
+	{
+		if (const UAssetDefinition* AssetDefinition = InAssetPayload.GetAssetDefinition())
+		{
+			if (AssetDefinition->ShouldSaveExternalPackages())
+			{
+				for (const UPackage* ExternalPackage : AssetPackage->GetExternalPackages())
+				{
+					if (ExternalPackage && ExternalPackage->IsDirty())
+					{
+						return true;
+					}
+				}
+			}
+		}
+
+		return AssetPackage->IsDirty();
+	}
+	
+	return false;
 }
 
 bool UpdateItemThumbnail(const UContentBrowserDataSource* InOwnerDataSource, const FContentBrowserItemData& InItem, FAssetThumbnail& InThumbnail)
diff --git a/Engine/Plugins/Experimental/Workspace/Source/WorkspaceEditor/Private/WorkspaceAssetEntry.cpp b/Engine/Plugins/Experimental/Workspace/Source/WorkspaceEditor/Private/WorkspaceAssetEntry.cpp
--- a/Engine/Plugins/Experimental/Workspace/Source/WorkspaceEditor/Private/WorkspaceAssetEntry.cpp
+++ b/Engine/Plugins/Experimental/Workspace/Source/WorkspaceEditor/Private/WorkspaceAssetEntry.cpp
@@ -18,4 +18,5 @@ void UWorkspaceAssetEntry::GetAssetRegistryTags(FAssetRegistryTagsContext Contex
 {
 	UObject::GetAssetRegistryTags(Context);
 	Context.AddTag(FAssetRegistryTag(UWorkspaceAssetEntry::ExportsAssetRegistryTag, Asset.GetUniqueID().ToString(), FAssetRegistryTag::TT_Hidden));
+	Context.AddTag(FAssetRegistryTag(FPrimaryAssetId::PrimaryAssetDisplayNameTag, *Asset.GetAssetName(), FAssetRegistryTag::TT_Hidden));
 }
diff --git a/Engine/Source/Editor/AssetDefinition/Public/AssetDefinition.h b/Engine/Source/Editor/AssetDefinition/Public/AssetDefinition.h
--- a/Engine/Source/Editor/AssetDefinition/Public/AssetDefinition.h
+++ b/Engine/Source/Editor/AssetDefinition/Public/AssetDefinition.h
@@ -566,7 +566,7 @@ public:
 		return TSharedPtr<SWidget>();
 	}
 
-	/** Whether or not external packages associated with this asset should be saved in addition to the asset itself */
+	/** Whether this asset has external packages associated with it (impacts saving and dirty-state behavior) */
 	virtual bool ShouldSaveExternalPackages() const
 	{
 		return false;
