Added a missing IsPrimaryAsset test when populating the content browser at a specific Path.

diff --git a/Engine/Plugins/Editor/ContentBrowser/ContentBrowserAssetDataSource/Source/ContentBrowserAssetDataSource/Private/ContentBrowserAssetDataSource.cpp b/Engine/Plugins/Editor/ContentBrowser/ContentBrowserAssetDataSource/Source/ContentBrowserAssetDataSource/Private/ContentBrowserAssetDataSource.cpp
--- a/Engine/Plugins/Editor/ContentBrowser/ContentBrowserAssetDataSource/Source/ContentBrowserAssetDataSource/Private/ContentBrowserAssetDataSource.cpp
+++ b/Engine/Plugins/Editor/ContentBrowser/ContentBrowserAssetDataSource/Source/ContentBrowserAssetDataSource/Private/ContentBrowserAssetDataSource.cpp
@@ -2118,7 +2118,7 @@ void UContentBrowserAssetDataSource::EnumerateItemsMatchingFilter(const FContent
 				Converted.Reserve(Assets.Num());
 				Converted.AddUninitialized(Assets.Num());
 				ParallelFor(TEXT("ConvertAssetsToContentBrowserItems"), Assets.Num(), 1024 * 16, [&Assets, &Converted, this](int32 Index) {
-					if (Assets[Index].IsValid())
+					if (Assets[Index].IsValid() && ContentBrowserAssetData::IsPrimaryAsset(Assets[Index]))
 					{
 						new (&Converted[Index]) FContentBrowserItemData(CreateAssetFileItem(MoveTemp(Assets[Index])));
 					}
@@ -2139,7 +2139,7 @@ void UContentBrowserAssetDataSource::EnumerateItemsMatchingFilter(const FContent
 			{
 				for (FAssetData& AssetData : Assets)
 				{
-					if (AssetData.IsValid())
+					if (AssetData.IsValid() && ContentBrowserAssetData::IsPrimaryAsset(AssetData))
 					{
 						InSink.ProduceItem(CreateAssetFileItem(MoveTemp(AssetData)));
 					}
@@ -3208,7 +3208,8 @@ void UContentBrowserAssetDataSource::OnAssetRegistryFileLoadProgress(const IAsse
 
 void UContentBrowserAssetDataSource::OnAssetsAdded(TConstArrayView<FAssetData> InAssets)
 {
-	TRACE_CPUPROFILER_EVENT_SCOPE(UContentBrowserAssetDataSource::OnAssetsAdded);
+	TRACE_CPUPROFILER_EVENT_SCOPE(UContentBrowserAssetDataSource::OnAssetAdded);
+	LLM_SCOPE_BYNAME(TEXT("UContentBrowserAssetDataSource"))
 
 	FAssetPropertyTagCache& Cache = FAssetPropertyTagCache::Get();
 	for (const FAssetData& InAssetData : InAssets)
@@ -3232,6 +3233,8 @@ void UContentBrowserAssetDataSource::OnAssetsAdded(TConstArrayView<FAssetData> I
 
 void UContentBrowserAssetDataSource::OnAssetRemoved(const FAssetData& InAssetData)
 {
+	TRACE_CPUPROFILER_EVENT_SCOPE(UContentBrowserAssetDataSource::OnAssetRemoved);
+	LLM_SCOPE_BYNAME(TEXT("UContentBrowserAssetDataSource"))
 	if (ContentBrowserAssetData::IsPrimaryAsset(InAssetData))
 	{
 		UE_LOG(LogContentBrowserAssetDataSource, VeryVerbose, TEXT("OnAssetRemoved: %s"), *WriteToString<256>(InAssetData.GetSoftObjectPath()));
@@ -3241,6 +3244,8 @@ void UContentBrowserAssetDataSource::OnAssetRemoved(const FAssetData& InAssetDat
 
 void UContentBrowserAssetDataSource::OnAssetRenamed(const FAssetData& InAssetData, const FString& InOldObjectPath)
 {
+	TRACE_CPUPROFILER_EVENT_SCOPE(UContentBrowserAssetDataSource::OnAssetRenamed);
+	LLM_SCOPE_BYNAME(TEXT("UContentBrowserAssetDataSource"))
 	if (ContentBrowserAssetData::IsPrimaryAsset(InAssetData))
 	{
 		UE_LOG(LogContentBrowserAssetDataSource, VeryVerbose, TEXT("OnAssetRenamed: %s"), *WriteToString<256>(InAssetData.GetSoftObjectPath()));
@@ -3257,6 +3262,8 @@ void UContentBrowserAssetDataSource::OnAssetRenamed(const FAssetData& InAssetDat
 
 void UContentBrowserAssetDataSource::OnAssetUpdated(const FAssetData& InAssetData)
 {
+	TRACE_CPUPROFILER_EVENT_SCOPE(UContentBrowserAssetDataSource::OnAssetUpdated);
+	LLM_SCOPE_BYNAME(TEXT("UContentBrowserAssetDataSource"))
 	if (ContentBrowserAssetData::IsPrimaryAsset(InAssetData))
 	{
 		UE_LOG(LogContentBrowserAssetDataSource, VeryVerbose, TEXT("OnAssetUpdated: %s"), *WriteToString<256>(InAssetData.GetSoftObjectPath()));
@@ -3268,6 +3275,8 @@ void UContentBrowserAssetDataSource::OnAssetUpdated(const FAssetData& InAssetDat
 
 void UContentBrowserAssetDataSource::OnAssetUpdatedOnDisk(const FAssetData& InAssetData)
 {
+	TRACE_CPUPROFILER_EVENT_SCOPE(UContentBrowserAssetDataSource::OnAssetUpdatedOnDisk);
+
 	if (ContentBrowserAssetData::IsPrimaryAsset(InAssetData))
 	{
 		UE_LOG(LogContentBrowserAssetDataSource, VeryVerbose, TEXT("OnAssetUpdatedOnDisk: %s"), *WriteToString<256>(InAssetData.GetSoftObjectPath()));
@@ -3279,6 +3288,8 @@ void UContentBrowserAssetDataSource::OnAssetUpdatedOnDisk(const FAssetData& InAs
 
 void UContentBrowserAssetDataSource::OnObjectPropertyChanged(UObject* InObject, FPropertyChangedEvent& InPropertyChangedEvent)
 {
+	TRACE_CPUPROFILER_EVENT_SCOPE(UContentBrowserAssetDataSource::OnObjectPropertyChanged);
+	LLM_SCOPE_BYNAME(TEXT("UContentBrowserAssetDataSource"))
 	if (InObject && InObject->IsAsset() && ContentBrowserAssetData::IsPrimaryAsset(InObject))
 	{
 		FAssetData AssetData(InObject);
@@ -3289,6 +3300,8 @@ void UContentBrowserAssetDataSource::OnObjectPropertyChanged(UObject* InObject,
 
 void UContentBrowserAssetDataSource::OnObjectPreSave(UObject* InObject, FObjectPreSaveContext InObjectPreSaveContext)
 {
+	TRACE_CPUPROFILER_EVENT_SCOPE(UContentBrowserAssetDataSource::OnObjectPreSave);
+	LLM_SCOPE_BYNAME(TEXT("UContentBrowserAssetDataSource"))
 	if (InObject && InObject->IsAsset() && ContentBrowserAssetData::IsPrimaryAsset(InObject))
 	{
 		FAssetData AssetData(InObject);
@@ -3299,6 +3312,8 @@ void UContentBrowserAssetDataSource::OnObjectPreSave(UObject* InObject, FObjectP
 
 void UContentBrowserAssetDataSource::OnPathsAdded(TConstArrayView<FStringView> Paths)
 {
+	TRACE_CPUPROFILER_EVENT_SCOPE(UContentBrowserAssetDataSource::OnPathsAdded);
+	LLM_SCOPE_BYNAME(TEXT("UContentBrowserAssetDataSource"))
 	RecentlyPopulatedAssetFolders.Empty();
 	for (FStringView InPath : Paths)
 	{
@@ -3333,6 +3348,8 @@ void UContentBrowserAssetDataSource::OnPathsAdded(TConstArrayView<FStringView> P
 
 void UContentBrowserAssetDataSource::OnPathsRemoved(TConstArrayView<FStringView> Paths)
 {
+	TRACE_CPUPROFILER_EVENT_SCOPE(UContentBrowserAssetDataSource::OnPathsRemoved);
+	LLM_SCOPE_BYNAME(TEXT("UContentBrowserAssetDataSource"))
 	for (FStringView InPath : Paths)
 	{
 		// Deleted paths are no longer relevant for tracking
