Fix content browser PathPermissionList not allowing /A/B when AllowList contains /A/B/C
#rnx

diff --git a/Engine/Plugins/Editor/ContentBrowser/ContentBrowserAssetDataSource/Source/ContentBrowserAssetDataSource/Private/ContentBrowserAssetDataSource.cpp b/Engine/Plugins/Editor/ContentBrowser/ContentBrowserAssetDataSource/Source/ContentBrowserAssetDataSource/Private/ContentBrowserAssetDataSource.cpp
--- a/Engine/Plugins/Editor/ContentBrowser/ContentBrowserAssetDataSource/Source/ContentBrowserAssetDataSource/Private/ContentBrowserAssetDataSource.cpp
+++ b/Engine/Plugins/Editor/ContentBrowser/ContentBrowserAssetDataSource/Source/ContentBrowserAssetDataSource/Private/ContentBrowserAssetDataSource.cpp
@@ -1884,7 +1884,7 @@ bool PathPassesCompiledDataFilterRecursive(const FContentBrowserCompiledAssetDat
 	if (!EnumHasAnyFlags(InOutFilterState, EFolderFilterState::SkipPermissionList))
 	{
 		EPathPermissionPrefixResult PermissionResult =
-			InFilter.PathPermissionList.PassesStartsWithFilterRecursive(Path);
+			InFilter.PathPermissionList.PassesStartsWithFilterRecursive(Path, /*bAllowParentPaths*/ true);
 		if (!UpdateFilterState(PermissionResult, EFolderFilterState::SkipPermissionList))
 		{
 			return false;
