Fix crash with content browser attempting to load package to build tooltip when right clicking assets
#rnx

diff --git a/Engine/Plugins/Editor/ContentBrowser/ContentBrowserAssetDataSource/Source/ContentBrowserAssetDataSource/Private/ContentBrowserAssetDataCore.cpp b/Engine/Plugins/Editor/ContentBrowser/ContentBrowserAssetDataSource/Source/ContentBrowserAssetDataSource/Private/ContentBrowserAssetDataCore.cpp
--- a/Engine/Plugins/Editor/ContentBrowser/ContentBrowserAssetDataSource/Source/ContentBrowserAssetDataSource/Private/ContentBrowserAssetDataCore.cpp
+++ b/Engine/Plugins/Editor/ContentBrowser/ContentBrowserAssetDataSource/Source/ContentBrowserAssetDataSource/Private/ContentBrowserAssetDataCore.cpp
@@ -1426,7 +1426,7 @@ void AppendAssetPackageName(const FAssetData& AssetData, FString& InOutStr)
 	{
 		InOutStr += LINE_TERMINATOR;
 	}
-	InOutStr += AssetData.PackageName.ToString();
+	AssetData.PackageName.AppendString(InOutStr);
 }
 
 bool AppendAssetFolderItemReference(IAssetRegistry* InAssetRegistry, const FContentBrowserAssetFolderItemDataPayload& InFolderPayload, FString& InOutStr)
