Content browser public asset support updated
#rnx
#tests Cooked and uncooked editor, preflighted fully cooked build

diff --git a/Engine/Plugins/Editor/ContentBrowser/ContentBrowserAssetDataSource/Source/ContentBrowserAssetDataSource/Private/ContentBrowserAssetDataSource.cpp b/Engine/Plugins/Editor/ContentBrowser/ContentBrowserAssetDataSource/Source/ContentBrowserAssetDataSource/Private/ContentBrowserAssetDataSource.cpp
--- a/Engine/Plugins/Editor/ContentBrowser/ContentBrowserAssetDataSource/Source/ContentBrowserAssetDataSource/Private/ContentBrowserAssetDataSource.cpp
+++ b/Engine/Plugins/Editor/ContentBrowser/ContentBrowserAssetDataSource/Source/ContentBrowserAssetDataSource/Private/ContentBrowserAssetDataSource.cpp
@@ -2056,14 +2056,25 @@ void UContentBrowserAssetDataSource::EnumerateItemsMatchingFilter(const FContent
 
 		if (const FContentBrowserCompiledUnsupportedAssetDataFilter* UnsupportedAssetDataFilter = FilterList->FindFilter<FContentBrowserCompiledUnsupportedAssetDataFilter>())
 		{
+			const TSharedPtr<FPathPermissionList>& ShowPrivateContentPermissionList = IContentBrowserSingleton::Get().GetShowPrivateContentPermissionList();
+
 			// Using the show unsupported asset filter
-			AssetRegistry->EnumerateAssets(UnsupportedAssetDataFilter->InclusiveFilter, [this, &InSink, &AssetDataFilter, UnsupportedAssetDataFilter](const FAssetData& AssetData)
+			AssetRegistry->EnumerateAssets(UnsupportedAssetDataFilter->InclusiveFilter, [this, &InSink, &AssetDataFilter, UnsupportedAssetDataFilter, &ShowPrivateContentPermissionList](const FAssetData& AssetData)
 			{
 				if (ContentBrowserAssetData::IsPrimaryAsset(AssetData) && AssetData.GetOptionalOuterPathName().IsNone())
 				{
 					const bool bPassesExclusiveFilter = UnsupportedAssetDataFilter->ExclusiveFilter.IsEmpty() || !AssetRegistry->IsAssetIncludedByFilter(AssetData, UnsupportedAssetDataFilter->ExclusiveFilter);
 					if (bPassesExclusiveFilter)
 					{
+						// Exclude private assets that do not pass ShowPrivateContentPermissionList
+						if (AssetData.HasAnyPackageFlags(PKG_NotExternallyReferenceable))
+						{
+							if (!ShowPrivateContentPermissionList->PassesStartsWithFilter(FNameBuilder(AssetData.PackageName)))
+							{
+								return true;
+							}
+						}
+
 						// Should this asset be presented as unsupported
 						if (!(AssetRegistry->IsAssetIncludedByFilter(AssetData, UnsupportedAssetDataFilter->ConvertIfFailInclusiveFilter) && (UnsupportedAssetDataFilter->ConvertIfFailExclusiveFilter.IsEmpty() || AssetRegistry->IsAssetExcludedByFilter(AssetData, UnsupportedAssetDataFilter->ConvertIfFailExclusiveFilter))) // Do we fail the supported filter?
 							&& (AssetRegistry->IsAssetIncludedByFilter(AssetData, UnsupportedAssetDataFilter->ShowInclusiveFilter) && (UnsupportedAssetDataFilter->ShowExclusiveFilter.IsEmpty() || AssetRegistry->IsAssetExcludedByFilter(AssetData, UnsupportedAssetDataFilter->ShowExclusiveFilter)))) // Do we pass the show filter for the unsupported asset?
@@ -2305,7 +2316,9 @@ bool UContentBrowserAssetDataSource::PrioritizeSearchPath(const FName InPath)
 
 bool UContentBrowserAssetDataSource::IsFolderVisible(const FName InPath, const EContentBrowserIsFolderVisibleFlags InFlags, TOptional<FContentBrowserFolderContentsFilter> InContentsFilter)
 {
-	auto IsInternalFolderVisible = [this, InContentsFilter](FName InternalFolderPath) {
+	const TSharedPtr<FPathPermissionList>& ShowPrivateContentPermissionList = IContentBrowserSingleton::Get().GetShowPrivateContentPermissionList();
+
+	auto IsInternalFolderVisible = [this, InContentsFilter, &ShowPrivateContentPermissionList](FName InternalFolderPath) {
 		const EContentBrowserFolderAttributes FolderAttributes = GetAssetFolderAttributes(InternalFolderPath);
 		if (EnumHasAnyFlags(FolderAttributes, EContentBrowserFolderAttributes::AlwaysVisible))
 		{
@@ -2313,10 +2326,16 @@ bool UContentBrowserAssetDataSource::IsFolderVisible(const FName InPath, const E
 		}
 
 		// Hide folders that only contain cooked private content
-		if (EnumHasAnyFlags(FolderAttributes, EContentBrowserFolderAttributes::HasAssets | EContentBrowserFolderAttributes::HasRedirectors)
-			&& !EnumHasAnyFlags(FolderAttributes, EContentBrowserFolderAttributes::HasPublicContent | EContentBrowserFolderAttributes::HasSourceContent))
+		if (EnumHasAnyFlags(FolderAttributes, EContentBrowserFolderAttributes::HasAssets | EContentBrowserFolderAttributes::HasRedirectors))
 		{
-			return false;
+			if (!EnumHasAnyFlags(FolderAttributes, EContentBrowserFolderAttributes::HasPublicContent))
+			{
+				FNameBuilder InternalFolderPathBuilder(InternalFolderPath);
+				if (!ShowPrivateContentPermissionList->PassesStartsWithFilter(InternalFolderPathBuilder))
+				{
+					return false;
+				}
+			}
 		}
 
 		if (InContentsFilter.IsSet())
@@ -2563,6 +2582,15 @@ bool UContentBrowserAssetDataSource::ConvertItemForFilter(FContentBrowserItemDat
 				if (!(AssetRegistry->IsAssetIncludedByFilter(AssetData, UnsupportedAssetDataFilter->ConvertIfFailInclusiveFilter) && (UnsupportedAssetDataFilter->ConvertIfFailExclusiveFilter.IsEmpty() || AssetRegistry->IsAssetExcludedByFilter(AssetData, UnsupportedAssetDataFilter->ConvertIfFailExclusiveFilter))) // Do we fail the supported filter?
 					&& (AssetRegistry->IsAssetIncludedByFilter(AssetData, UnsupportedAssetDataFilter->ShowInclusiveFilter) && (UnsupportedAssetDataFilter->ShowExclusiveFilter.IsEmpty() || AssetRegistry->IsAssetExcludedByFilter(AssetData, UnsupportedAssetDataFilter->ShowExclusiveFilter)))) // Do we pass the show filter for the unsupported asset?
 				{
+					// Exclude private assets that do not pass ShowPrivateContentPermissionList
+					if (AssetData.HasAnyPackageFlags(PKG_NotExternallyReferenceable))
+					{
+						if (!IContentBrowserSingleton::Get().GetShowPrivateContentPermissionList()->PassesStartsWithFilter(FNameBuilder(AssetData.PackageName)))
+						{
+							return false;
+						}
+					}
+
 					Item = CreateUnsupportedAssetFileItem(AssetData);
 					return true;
 				}
diff --git a/Engine/Source/Editor/ContentBrowser/Private/SContentBrowser.cpp b/Engine/Source/Editor/ContentBrowser/Private/SContentBrowser.cpp
--- a/Engine/Source/Editor/ContentBrowser/Private/SContentBrowser.cpp
+++ b/Engine/Source/Editor/ContentBrowser/Private/SContentBrowser.cpp
@@ -463,9 +463,6 @@ void SContentBrowser::Construct( const FArguments& InArgs, const FName& InInstan
 			.OnItemsActivated(this, &SContentBrowser::OnItemsActivated)
 			.OnGetItemContextMenu(this, &SContentBrowser::GetItemContextMenu, EContentBrowserViewContext::AssetView)
 			.OnItemRenameCommitted(this, &SContentBrowser::OnItemRenameCommitted)
-			.OnShouldFilterItem(bPrivateContentFilterEnabled
-				? FOnShouldFilterItem::CreateSP(this, &SContentBrowser::HandlePrivateContentFilter)
-				: FOnShouldFilterItem())
 			.FrontendFilters(FrontendFilters)
 			.TextFilter(TextFilter)
 			.ShowRedirectors_Lambda([this]() { return ContentBrowserUtils::ShouldShowRedirectors(FilterListPtr); })
@@ -3685,20 +3682,6 @@ void SContentBrowser::OnConsoleVariableChanged()
 
 void SContentBrowser::UpdatePrivateContentFeatureEnabled(bool bUpdateFilterIfChanged)
 {
-	static const IConsoleVariable* EnablePublicAssetFeatureCVar =
-		IConsoleManager::Get().FindConsoleVariable(TEXT("AssetTools.EnablePublicAssetFeature"));
-	const bool bShouldEnablePrivateContentFilter = EnablePublicAssetFeatureCVar
-		&& EnablePublicAssetFeatureCVar->GetBool();
-	if (bPrivateContentFilterEnabled != bShouldEnablePrivateContentFilter)
-	{
-		if (bUpdateFilterIfChanged)
-		{
-			AssetViewPtr->SetShouldFilterItem(bShouldEnablePrivateContentFilter
-				? FOnShouldFilterItem::CreateSP(this, &SContentBrowser::HandlePrivateContentFilter)
-				: FOnShouldFilterItem());
-		}
-		bPrivateContentFilterEnabled = bShouldEnablePrivateContentFilter;
-	}
 }
 
 FReply SContentBrowser::BackClicked()
@@ -4413,25 +4396,6 @@ void SContentBrowser::HandleItemDataUpdated(TArrayView<const FContentBrowserItem
 	}
 }
 
-bool SContentBrowser::HandlePrivateContentFilter(const FContentBrowserItem& AssetItem)
-{
-	static const IConsoleVariable* EnablePublicAssetFeatureCVar = IConsoleManager::Get().FindConsoleVariable(TEXT("AssetTools.EnablePublicAssetFeature"));
-	if (!EnablePublicAssetFeatureCVar || !EnablePublicAssetFeatureCVar->GetBool())
-	{
-		return false;
-	}
-
-	FAssetData ItemAssetData;
-	if ((AssetItem.Legacy_TryGetAssetData(ItemAssetData) && (ItemAssetData.PackageFlags & PKG_NotExternallyReferenceable)))
-	{
-		const FNameBuilder AssetItemFolderPath(AssetItem.GetVirtualPath());
-
-		return !FContentBrowserSingleton::Get().IsShowingPrivateContent(AssetItemFolderPath);
-	}
-
-	return false;
-}
-
 FText SContentBrowser::GetSearchAssetsHintText() const
 {
 	if (PathViewPtr.IsValid())
