Fix AssetFileContextMenu for aliases

diff --git a/Engine/Plugins/Editor/ContentBrowser/ContentBrowserAssetDataSource/Source/ContentBrowserAssetDataSource/Private/AssetFileContextMenu.cpp b/Engine/Plugins/Editor/ContentBrowser/ContentBrowserAssetDataSource/Source/ContentBrowserAssetDataSource/Private/AssetFileContextMenu.cpp
--- a/Engine/Plugins/Editor/ContentBrowser/ContentBrowserAssetDataSource/Source/ContentBrowserAssetDataSource/Private/AssetFileContextMenu.cpp
+++ b/Engine/Plugins/Editor/ContentBrowser/ContentBrowserAssetDataSource/Source/ContentBrowserAssetDataSource/Private/AssetFileContextMenu.cpp
@@ -268,7 +268,6 @@ bool FAssetFileContextMenu::AddCommonMenuOptions(UToolMenu* Menu)
 	UContentBrowserDataMenuContext_FileMenu* Context = Menu->FindContext<UContentBrowserDataMenuContext_FileMenu>();
 	const bool bCanBeModified = !Context || Context->bCanBeModified;
 
-	if (bCanBeModified)
 	{
 		FToolMenuSection& Section = Menu->AddSection("CommonAssetActions", LOCTEXT("CommonAssetActionsMenuHeading", "Common"));
 
@@ -287,17 +286,20 @@ bool FAssetFileContextMenu::AddCommonMenuOptions(UToolMenu* Menu)
 			FSlateIcon(FAppStyle::GetAppStyleSetName(), "Icons.Adjust")
 			);
 
-		// Asset Localization sub-menu
-		Section.AddSubMenu(
-			"LocalizationSubMenu",
-			LOCTEXT("LocalizationSubMenuLabel", "Asset Localization"),
-			LOCTEXT("LocalizationSubMenuToolTip", "Manage the localization of this asset"),
-			FNewToolMenuDelegate::CreateSP(this, &FAssetFileContextMenu::MakeAssetLocalizationSubMenu),
-			FUIAction(),
-			EUserInterfaceActionType::Button,
-			false,
-			FSlateIcon(FAppStyle::GetAppStyleSetName(), "Icons.Localization")
+		if (bCanBeModified)
+		{
+			// Asset Localization sub-menu
+			Section.AddSubMenu(
+				"LocalizationSubMenu",
+				LOCTEXT("LocalizationSubMenuLabel", "Asset Localization"),
+				LOCTEXT("LocalizationSubMenuToolTip", "Manage the localization of this asset"),
+				FNewToolMenuDelegate::CreateSP(this, &FAssetFileContextMenu::MakeAssetLocalizationSubMenu),
+				FUIAction(),
+				EUserInterfaceActionType::Button,
+				false,
+				FSlateIcon(FAppStyle::GetAppStyleSetName(), "Icons.Localization")
 			);
+		}
 	}
 
 	return true;
@@ -414,7 +416,6 @@ void FAssetFileContextMenu::MakeAssetActionsSubMenu(UToolMenu* Menu)
 	}
 
 	// ADVANCED ACTIONS
-	if (bCanBeModified)
 	{
 		FToolMenuSection& Section = Menu->AddSection("AssetContextAdvancedActions", LOCTEXT("AssetContextAdvancedActionsMenuHeading", "Advanced"));
 
@@ -454,14 +455,17 @@ void FAssetFileContextMenu::MakeAssetActionsSubMenu(UToolMenu* Menu)
 		);
 
 		// Property Matrix
-		bool bCanUsePropertyMatrix = FModuleManager::LoadModuleChecked<FPropertyEditorModule>( "PropertyEditor" ).GetCanUsePropertyMatrix();
-		// Materials can't be bulk edited currently as they require very special handling because of their dependencies with the rendering thread, and we'd have to hack the property matrix too much.
-		for (auto& Asset : SelectedAssets)
+		bool bCanUsePropertyMatrix = bCanBeModified && FModuleManager::LoadModuleChecked<FPropertyEditorModule>( "PropertyEditor" ).GetCanUsePropertyMatrix();
+		if (bCanUsePropertyMatrix)
 		{
-			if (Asset.AssetClassPath == UMaterial::StaticClass()->GetClassPathName() || Asset.AssetClassPath == UMaterialInstanceConstant::StaticClass()->GetClassPathName() || Asset.AssetClassPath == UMaterialFunction::StaticClass()->GetClassPathName() || Asset.AssetClassPath == UMaterialFunctionInstance::StaticClass()->GetClassPathName())
+			// Materials can't be bulk edited currently as they require very special handling because of their dependencies with the rendering thread, and we'd have to hack the property matrix too much.
+			for (const FAssetData& Asset : SelectedAssets)
 			{
-				bCanUsePropertyMatrix = false;
-				break;
+				if (Asset.AssetClassPath == UMaterial::StaticClass()->GetClassPathName() || Asset.AssetClassPath == UMaterialInstanceConstant::StaticClass()->GetClassPathName() || Asset.AssetClassPath == UMaterialFunction::StaticClass()->GetClassPathName() || Asset.AssetClassPath == UMaterialFunctionInstance::StaticClass()->GetClassPathName())
+				{
+					bCanUsePropertyMatrix = false;
+					break;
+				}
 			}
 		}
 
@@ -483,17 +487,20 @@ void FAssetFileContextMenu::MakeAssetActionsSubMenu(UToolMenu* Menu)
 				);
 		}
 
-		// Create Metadata menu
-		Section.AddMenuEntry(
-			"ShowAssetMetaData",
-			LOCTEXT("ShowAssetMetaData", "Show Metadata"),
-			LOCTEXT("ShowAssetMetaDataTooltip", "Show the asset metadata dialog."),
-			FSlateIcon(),
-			FUIAction( FExecuteAction::CreateSP(this, &FAssetFileContextMenu::ExecuteShowAssetMetaData) )
-		);
+		if (bCanBeModified)
+		{
+			// Create Metadata menu
+			Section.AddMenuEntry(
+				"ShowAssetMetaData",
+				LOCTEXT("ShowAssetMetaData", "Show Metadata"),
+				LOCTEXT("ShowAssetMetaDataTooltip", "Show the asset metadata dialog."),
+				FSlateIcon(),
+				FUIAction( FExecuteAction::CreateSP(this, &FAssetFileContextMenu::ExecuteShowAssetMetaData) )
+			);
+		}
 
 		// Chunk actions
-		if (GetDefault<UEditorExperimentalSettings>()->bContextMenuChunkAssignments)
+		if (bCanBeModified && GetDefault<UEditorExperimentalSettings>()->bContextMenuChunkAssignments)
 		{
 			Section.AddMenuEntry(
 				"AssignAssetChunk",
diff --git a/Engine/Plugins/Editor/ContentBrowser/ContentBrowserAssetDataSource/Source/ContentBrowserAssetDataSource/Private/ContentBrowserAssetDataCore.cpp b/Engine/Plugins/Editor/ContentBrowser/ContentBrowserAssetDataSource/Source/ContentBrowserAssetDataSource/Private/ContentBrowserAssetDataCore.cpp
--- a/Engine/Plugins/Editor/ContentBrowser/ContentBrowserAssetDataSource/Source/ContentBrowserAssetDataSource/Private/ContentBrowserAssetDataCore.cpp
+++ b/Engine/Plugins/Editor/ContentBrowser/ContentBrowserAssetDataSource/Source/ContentBrowserAssetDataSource/Private/ContentBrowserAssetDataCore.cpp
@@ -13,6 +13,7 @@
 #include "ContentBrowserDataSource.h"
 #include "ContentBrowserDataSubsystem.h"
 #include "ContentBrowserDataUtils.h"
+#include "ContentBrowserMenuContexts.h"
 #include "Editor.h"
 #include "Editor/UnrealEdTypes.h"
 #include "Engine/Level.h"
@@ -2063,18 +2064,8 @@ void PopulateAssetFileContextMenu(UContentBrowserDataSource* InOwnerDataSource,
 	const UContentBrowserDataMenuContext_FileMenu* ContextObject = InMenu->FindContext<UContentBrowserDataMenuContext_FileMenu>();
 	checkf(ContextObject, TEXT("Required context UContentBrowserDataMenuContext_FileMenu was missing!"));
 
-	// Extract the internal asset data that belong to this data source from the full list of selected items given in the context
-	TArray<FAssetData> SelectedAssets;
-	for (const FContentBrowserItem& SelectedItem : ContextObject->SelectedItems)
-	{
-		if (const FContentBrowserItemData* SelectedItemData = SelectedItem.GetPrimaryInternalItem())
-		{
-			if (TSharedPtr<const FContentBrowserAssetFileItemDataPayload> AssetPayload = GetAssetFileItemPayload(InOwnerDataSource, *SelectedItemData))
-			{
-				SelectedAssets.Add(AssetPayload->GetAssetData());
-			}
-		}
-	}
+	const UContentBrowserAssetContextMenuContext* AssetContextMenu = InMenu->FindContext<UContentBrowserAssetContextMenuContext>();
+	checkf(AssetContextMenu, TEXT("Required context UContentBrowserAssetContextMenuContext was missing!"));
 
 	FAssetFileContextMenu::FOnShowAssetsInPathsView OnShowAssetsInPathsView = FAssetFileContextMenu::FOnShowAssetsInPathsView::CreateLambda([OwnerDataSource = TWeakObjectPtr<UContentBrowserDataSource>(InOwnerDataSource), OnShowInPathsView = ContextObject->OnShowInPathsView](const TArray<FAssetData>& InAssetsToShow)
 	{
@@ -2099,7 +2090,7 @@ void PopulateAssetFileContextMenu(UContentBrowserDataSource* InOwnerDataSource,
 
 	InAssetFileContextMenu.MakeContextMenu(
 		InMenu,
-		SelectedAssets,
+		AssetContextMenu->SelectedAssets,
 		OnShowAssetsInPathsView
 	);
 }
diff --git a/Engine/Plugins/Editor/DataValidation/Source/DataValidation/Private/DataValidationModule.cpp b/Engine/Plugins/Editor/DataValidation/Source/DataValidation/Private/DataValidationModule.cpp
--- a/Engine/Plugins/Editor/DataValidation/Source/DataValidation/Private/DataValidationModule.cpp
+++ b/Engine/Plugins/Editor/DataValidation/Source/DataValidation/Private/DataValidationModule.cpp
@@ -170,34 +170,41 @@ void FDataValidationModule::RegisterMenus()
 
 	{
 		UToolMenu* Menu = UToolMenus::Get()->ExtendMenu("ContentBrowser.AssetContextMenu.AssetActionsSubMenu");
-		FToolMenuSection& Section = Menu->FindOrAddSection("AssetContextAdvancedActions");
-		Section.AddMenuEntry(
-			"ValidateAssets",
-			LOCTEXT("ValidateAssetsTabTitle", "Validate Assets"),
-			LOCTEXT("ValidateAssetsTooltipText", "Runs data validation on these assets."),
-			FSlateIcon(),
-			FToolMenuExecuteAction::CreateLambda([this](const FToolMenuContext& InContext)
-			{
-				if (const UContentBrowserAssetContextMenuContext* Context = InContext.FindContext<UContentBrowserAssetContextMenuContext>())
-				{
-					ValidateAssets(Context->SelectedAssets, false, EDataValidationUsecase::Manual);
-				}
-			})
-		);
-
-		Section.AddMenuEntry(
-			"ValidateAssetsAndDependencies",
-			LOCTEXT("ValidateAssetsAndDependenciesTabTitle", "Validate Assets and Dependencies"),
-			LOCTEXT("ValidateAssetsAndDependenciesTooltipText", "Runs data validation on these assets and all assets they depend on."),
-			FSlateIcon(),
-			FToolMenuExecuteAction::CreateLambda([this](const FToolMenuContext& InContext)
+		Menu->AddDynamicSection(NAME_None, FNewToolMenuDelegate::CreateLambda([this](UToolMenu* InMenu)
+		{
+			const UContentBrowserAssetContextMenuContext* Context = InMenu->FindContext<UContentBrowserAssetContextMenuContext>();
+			if (Context && Context->bCanBeModified && Context->SelectedAssets.Num() > 0)
 			{
-				if (const UContentBrowserAssetContextMenuContext* Context = InContext.FindContext<UContentBrowserAssetContextMenuContext>())
-				{
-					ValidateAssets(Context->SelectedAssets, true, EDataValidationUsecase::Manual);
-				}
-			})
-		);
+				FToolMenuSection& Section = InMenu->FindOrAddSection("AssetContextAdvancedActions");
+				Section.AddMenuEntry(
+					"ValidateAssets",
+					LOCTEXT("ValidateAssetsTabTitle", "Validate Assets"),
+					LOCTEXT("ValidateAssetsTooltipText", "Runs data validation on these assets."),
+					FSlateIcon(),
+					FToolMenuExecuteAction::CreateLambda([this](const FToolMenuContext& InContext)
+					{
+						if (const UContentBrowserAssetContextMenuContext* Context = InContext.FindContext<UContentBrowserAssetContextMenuContext>())
+						{
+							ValidateAssets(Context->SelectedAssets, false, EDataValidationUsecase::Manual);
+						}
+					})
+				);
+
+				Section.AddMenuEntry(
+					"ValidateAssetsAndDependencies",
+					LOCTEXT("ValidateAssetsAndDependenciesTabTitle", "Validate Assets and Dependencies"),
+					LOCTEXT("ValidateAssetsAndDependenciesTooltipText", "Runs data validation on these assets and all assets they depend on."),
+					FSlateIcon(),
+					FToolMenuExecuteAction::CreateLambda([this](const FToolMenuContext& InContext)
+					{
+						if (const UContentBrowserAssetContextMenuContext* Context = InContext.FindContext<UContentBrowserAssetContextMenuContext>())
+						{
+							ValidateAssets(Context->SelectedAssets, true, EDataValidationUsecase::Manual);
+						}
+					})
+				);
+			}
+		}));
 	}
 
 	{
