Hide 'right click asset in content browser > Asset Actions > Replace References' behind bCanBeModified
This command combines the selected assets into one of the selected assets and can delete the other assets selected.
This requires the ability to modify the selected assets.
#rnx

diff --git a/Engine/Plugins/Editor/ContentBrowser/ContentBrowserAssetDataSource/Source/ContentBrowserAssetDataSource/Private/AssetFileContextMenu.cpp b/Engine/Plugins/Editor/ContentBrowser/ContentBrowserAssetDataSource/Source/ContentBrowserAssetDataSource/Private/AssetFileContextMenu.cpp
--- a/Engine/Plugins/Editor/ContentBrowser/ContentBrowserAssetDataSource/Source/ContentBrowserAssetDataSource/Private/AssetFileContextMenu.cpp
+++ b/Engine/Plugins/Editor/ContentBrowser/ContentBrowserAssetDataSource/Source/ContentBrowserAssetDataSource/Private/AssetFileContextMenu.cpp
@@ -444,15 +444,18 @@ void FAssetFileContextMenu::MakeAssetActionsSubMenu(UToolMenu* Menu)
 		);
 
 		// Replace References
-		Section.AddMenuEntry(
-			"ReplaceReferences",
-			LOCTEXT("ReplaceReferences", "Replace References"),
-			LOCTEXT("ConsolidateTooltip", "Replace references to the selected assets."),
-			FSlateIcon(),
-			FUIAction(
-			FExecuteAction::CreateSP(this, &FAssetFileContextMenu::ExecuteConsolidate)
-			)
-		);
+		if (bCanBeModified)
+		{
+			Section.AddMenuEntry(
+				"ReplaceReferences",
+				LOCTEXT("ReplaceReferences", "Replace References"),
+				LOCTEXT("ConsolidateTooltip", "Replace references to the selected assets."),
+				FSlateIcon(),
+				FUIAction(
+				FExecuteAction::CreateSP(this, &FAssetFileContextMenu::ExecuteConsolidate)
+				)
+			);
+		}
 
 		// Property Matrix
 		bool bCanUsePropertyMatrix = bCanBeModified && FModuleManager::LoadModuleChecked<FPropertyEditorModule>( "PropertyEditor" ).GetCanUsePropertyMatrix();
