Gave name to factory created menus so they can be extended

#rnx

diff --git a/Engine/Plugins/Editor/ContentBrowser/ContentBrowserAssetDataSource/Source/ContentBrowserAssetDataSource/Private/NewAssetContextMenu.cpp b/Engine/Plugins/Editor/ContentBrowser/ContentBrowserAssetDataSource/Source/ContentBrowserAssetDataSource/Private/NewAssetContextMenu.cpp
--- a/Engine/Plugins/Editor/ContentBrowser/ContentBrowserAssetDataSource/Source/ContentBrowserAssetDataSource/Private/NewAssetContextMenu.cpp
+++ b/Engine/Plugins/Editor/ContentBrowser/ContentBrowserAssetDataSource/Source/ContentBrowserAssetDataSource/Private/NewAssetContextMenu.cpp
@@ -319,7 +319,7 @@ void FNewAssetContextMenu::MakeContextMenu(
 				if (Factories.Num() > 0)
 				{
 					Section.AddSubMenu(
-						NAME_None,
+						FName(AdvancedAssetCategory.CategoryName.BuildSourceString()),
 						AdvancedAssetCategory.CategoryName,
 						FText::GetEmpty(),
 						FNewToolMenuDelegate::CreateStatic(
