Added support for showing aliased content when filtering by collections in the content browser.

diff --git a/Engine/Plugins/Editor/ContentBrowser/ContentBrowserAliasDataSource/Source/ContentBrowserAliasDataSource/Private/ContentBrowserAliasDataSource.cpp b/Engine/Plugins/Editor/ContentBrowser/ContentBrowserAliasDataSource/Source/ContentBrowserAliasDataSource/Private/ContentBrowserAliasDataSource.cpp
--- a/Engine/Plugins/Editor/ContentBrowser/ContentBrowserAliasDataSource/Source/ContentBrowserAliasDataSource/Private/ContentBrowserAliasDataSource.cpp
+++ b/Engine/Plugins/Editor/ContentBrowser/ContentBrowserAliasDataSource/Source/ContentBrowserAliasDataSource/Private/ContentBrowserAliasDataSource.cpp
@@ -11,6 +11,7 @@
 #include "AssetToolsModule.h"
 #include "CollectionManagerModule.h"
 #include "IAssetTools.h"
+#include "ICollectionContainer.h"
 #include "IContentBrowserDataModule.h"
 #include "Misc/PackageName.h"
 
@@ -130,8 +131,34 @@ void UContentBrowserAliasDataSource::CompileFilter(const FName InPath, const FCo
 					PathTree.EnumerateSubPaths(Path, Callback, bRecursive);
 				};
 
+			auto CustomCollectionEnumeration = [this](const FCollectionRef& Collection, ECollectionRecursionFlags::Flags RecursionMode, TFunctionRef<void(const FSoftObjectPath&)> Callback)
+				{
+					if (!bFilterShouldMatchCollectionContent)
+					{
+						return;
+					}
+
+					TArray<FSoftObjectPath> ObjectPaths;
+					Collection.Container->GetObjectsInCollection(Collection.Name, Collection.Type, ObjectPaths, RecursionMode);
+
+					for (const FSoftObjectPath& ObjectPath : ObjectPaths)
+					{
+						if (const TArray<FName>* FoundAliases = AliasesForObjectPath.Find(ObjectPath))
+						{
+							for (const FName Alias : *FoundAliases)
+							{
+								if (const FAliasData* AliasData = AllAliases.Find(FContentBrowserUniqueAlias(ObjectPath, Alias)))
+								{
+									Callback(FSoftObjectPath::ConstructFromPackageAsset(AliasData->PackageName, AliasData->AssetData.AssetName));
+								}
+							}
+						}
+					}
+				};
+
 			UContentBrowserAssetDataSource::FSubPathEnumerationFunc CustomSubPathEnumerationRef = CustomSubPathEnumeration;
-			UContentBrowserAssetDataSource::CreateAssetFilter(Params, InPath, InFilter, OutCompiledFilter, &CustomSubPathEnumerationRef);
+			UContentBrowserAssetDataSource::FCollectionEnumerationFunc CustomCollectionEnumerationRef = CustomCollectionEnumeration;
+			UContentBrowserAssetDataSource::CreateAssetFilter(Params, InPath, InFilter, OutCompiledFilter, &CustomSubPathEnumerationRef, &CustomCollectionEnumerationRef);
 		}
 	}
 }
@@ -750,6 +777,11 @@ void UContentBrowserAliasDataSource::LogAliases() const
 	}
 }
 
+void UContentBrowserAliasDataSource::SetFilterShouldMatchCollectionContent(bool bInFilterShouldMatchCollectionContent)
+{
+	bFilterShouldMatchCollectionContent = bInFilterShouldMatchCollectionContent;
+}
+
 void UContentBrowserAliasDataSource::OnAssetUpdated(const FAssetData& InAssetData)
 {
 	ReconcileAliasesFromMetaData(InAssetData);
@@ -955,6 +987,11 @@ bool UContentBrowserAliasDataSource::UpdateThumbnail(const FContentBrowserItemDa
 
 bool UContentBrowserAliasDataSource::TryGetCollectionId(const FContentBrowserItemData& InItem, FSoftObjectPath& OutCollectionId)
 {
+	if (TSharedPtr<const FContentBrowserAssetFileItemDataPayload> AssetPayload = ContentBrowserAssetData::GetAssetFileItemPayload(this, InItem))
+	{
+		OutCollectionId = AssetPayload->GetAssetData().GetSoftObjectPath();
+		return true;
+	}
 	return false;
 }
 
diff --git a/Engine/Plugins/Editor/ContentBrowser/ContentBrowserAliasDataSource/Source/ContentBrowserAliasDataSource/Public/ContentBrowserAliasDataSource.h b/Engine/Plugins/Editor/ContentBrowser/ContentBrowserAliasDataSource/Source/ContentBrowserAliasDataSource/Public/ContentBrowserAliasDataSource.h
--- a/Engine/Plugins/Editor/ContentBrowser/ContentBrowserAliasDataSource/Source/ContentBrowserAliasDataSource/Public/ContentBrowserAliasDataSource.h
+++ b/Engine/Plugins/Editor/ContentBrowser/ContentBrowserAliasDataSource/Source/ContentBrowserAliasDataSource/Public/ContentBrowserAliasDataSource.h
@@ -158,6 +158,8 @@ public:
 	/** Logs all the content browser aliases */
 	void LogAliases() const;
 
+	void SetFilterShouldMatchCollectionContent(bool bInFilterShouldMatchCollectionContent);
+
 protected:
 	virtual void BuildRootPathVirtualTree() override;
 
@@ -242,4 +244,7 @@ private:
 
 	/** Delegate broadcast after all aliases removed to give chance for systems to re-add aliases */
 	FSimpleMulticastDelegate RebuildAliasesDelegate;
+
+	/** If true, EnumerateItemsMatchingFilter will include aliases of items contained in collections included in the filter */
+	bool bFilterShouldMatchCollectionContent = true;
 };
diff --git a/Engine/Plugins/Editor/ContentBrowser/ContentBrowserAssetDataSource/Source/ContentBrowserAssetDataSource/Private/ContentBrowserAssetDataSource.cpp b/Engine/Plugins/Editor/ContentBrowser/ContentBrowserAssetDataSource/Source/ContentBrowserAssetDataSource/Private/ContentBrowserAssetDataSource.cpp
--- a/Engine/Plugins/Editor/ContentBrowser/ContentBrowserAssetDataSource/Source/ContentBrowserAssetDataSource/Private/ContentBrowserAssetDataSource.cpp
+++ b/Engine/Plugins/Editor/ContentBrowser/ContentBrowserAssetDataSource/Source/ContentBrowserAssetDataSource/Private/ContentBrowserAssetDataSource.cpp
@@ -567,7 +567,7 @@ bool UContentBrowserAssetDataSource::CreatePathFilter(FAssetFilterInputParams& P
 	return true;
 }
 
-bool UContentBrowserAssetDataSource::CreateAssetFilter(FAssetFilterInputParams& Params, FName InPath, const FContentBrowserDataFilter& InFilter, FContentBrowserDataCompiledFilter& OutCompiledFilter, FSubPathEnumerationFunc* GetSubPackagePathsFunc)
+bool UContentBrowserAssetDataSource::CreateAssetFilter(FAssetFilterInputParams& Params, FName InPath, const FContentBrowserDataFilter& InFilter, FContentBrowserDataCompiledFilter& OutCompiledFilter, FSubPathEnumerationFunc* GetSubPackagePathsFunc, FCollectionEnumerationFunc* GetCollectionObjectPathsFunc)
 {
 	TRACE_CPUPROFILER_EVENT_SCOPE(UContentBrowserAssetDataSource::CreateAssetFilter);
 
@@ -651,7 +651,8 @@ PRAGMA_ENABLE_DEPRECATION_WARNINGS
 			if (Params.CollectionFilter)
 			{
 				TArray<FSoftObjectPath> ObjectPathsForCollections;
-				if (GetObjectPathsForCollections(Params.CollectionFilter->Collections, Params.CollectionFilter->bIncludeChildCollections, ObjectPathsForCollections) && ObjectPathsForCollections.Num() == 0)
+				if (GetObjectPathsForCollections(Params.CollectionFilter->Collections, Params.CollectionFilter->bIncludeChildCollections, GetCollectionObjectPathsFunc, ObjectPathsForCollections) &&
+					ObjectPathsForCollections.Num() == 0)
 				{
 					// If we had collections but they contained no objects then we can bail as nothing will pass the filter
 					return false;
@@ -1285,7 +1286,8 @@ PRAGMA_ENABLE_DEPRECATION_WARNINGS
 			if (Params.CollectionFilter)
 			{
 				TArray<FSoftObjectPath> ObjectPathsForCollections;
-				if (GetObjectPathsForCollections(Params.CollectionFilter->Collections, Params.CollectionFilter->bIncludeChildCollections, ObjectPathsForCollections) && ObjectPathsForCollections.Num() == 0)
+				if (GetObjectPathsForCollections(Params.CollectionFilter->Collections, Params.CollectionFilter->bIncludeChildCollections, nullptr, ObjectPathsForCollections) &&
+					ObjectPathsForCollections.Num() == 0)
 				{
 					// If we had collections but they contained no objects then we can bail as nothing will pass the filter
 					return false;
@@ -3143,7 +3145,7 @@ bool UContentBrowserAssetDataSource::IsKnownContentPath(const FName InPackagePat
 	return CurrentNode->NextNodes.Contains(TEXT('/'));
 }
 
-bool UContentBrowserAssetDataSource::GetObjectPathsForCollections(TArrayView<const FCollectionRef> InCollections, const bool bIncludeChildCollections, TArray<FSoftObjectPath>& OutObjectPaths)
+bool UContentBrowserAssetDataSource::GetObjectPathsForCollections(TArrayView<const FCollectionRef> InCollections, const bool bIncludeChildCollections, FCollectionEnumerationFunc* GetCollectionObjectPathsFunc, TArray<FSoftObjectPath>& OutObjectPaths)
 {
 	if (InCollections.Num() > 0)
 	{
@@ -3153,7 +3155,17 @@ bool UContentBrowserAssetDataSource::GetObjectPathsForCollections(TArrayView<con
 		{
 			if (Collection.Container)
 			{
-				Collection.Container->GetObjectsInCollection(Collection.Name, Collection.Type, OutObjectPaths, CollectionRecursionMode);
+				if (GetCollectionObjectPathsFunc != nullptr)
+				{
+					(*GetCollectionObjectPathsFunc)(Collection, CollectionRecursionMode, [&OutObjectPaths](const FSoftObjectPath& ObjectPath)
+						{
+							OutObjectPaths.Add(ObjectPath);
+						});
+				}
+				else
+				{
+					Collection.Container->GetObjectsInCollection(Collection.Name, Collection.Type, OutObjectPaths, CollectionRecursionMode);
+				}
 			}
 		}
 
diff --git a/Engine/Plugins/Editor/ContentBrowser/ContentBrowserAssetDataSource/Source/ContentBrowserAssetDataSource/Public/ContentBrowserAssetDataSource.h b/Engine/Plugins/Editor/ContentBrowser/ContentBrowserAssetDataSource/Source/ContentBrowserAssetDataSource/Public/ContentBrowserAssetDataSource.h
--- a/Engine/Plugins/Editor/ContentBrowser/ContentBrowserAssetDataSource/Source/ContentBrowserAssetDataSource/Public/ContentBrowserAssetDataSource.h
+++ b/Engine/Plugins/Editor/ContentBrowser/ContentBrowserAssetDataSource/Source/ContentBrowserAssetDataSource/Public/ContentBrowserAssetDataSource.h
@@ -164,6 +164,7 @@ public:
 		bool bIncludeRedirectors = false;
 	};
 
+	typedef TFunctionRef<void(const FCollectionRef&, ECollectionRecursionFlags::Flags, TFunctionRef<void(const FSoftObjectPath&)>)> FCollectionEnumerationFunc;
 	typedef TFunctionRef<void(FName, TFunctionRef<bool(FName)>, bool)> FSubPathEnumerationFunc;
 	typedef TFunctionRef<void(FARFilter&, FARCompiledFilter&)>         FCompileARFilterFunc;
 	typedef TFunctionRef<FContentBrowserItemData(FName)>               FCreateFolderItemFunc;
@@ -206,7 +207,7 @@ public:
 	 *
 	 * @return false if it's not possible to display any assets, otherwise true
 	 */
-	static bool CreateAssetFilter(FAssetFilterInputParams& Params, FName InPath, const FContentBrowserDataFilter& InFilter, FContentBrowserDataCompiledFilter& OutCompiledFilter, FSubPathEnumerationFunc* GetSubPackagePathsFunc = nullptr);
+	static bool CreateAssetFilter(FAssetFilterInputParams& Params, FName InPath, const FContentBrowserDataFilter& InFilter, FContentBrowserDataCompiledFilter& OutCompiledFilter, FSubPathEnumerationFunc* GetSubPackagePathsFunc = nullptr, FCollectionEnumerationFunc* GetCollectionObjectPathsFunc = nullptr);
 
 	UE_DEPRECATED(5.3, "Use the override without the function to customize the compilation the ARFilter or duplicate this function logic. This function will no longer be supported and will be removed to allow future performence improvements")
 	static bool CreateAssetFilter(FAssetFilterInputParams& Params, FName InPath, const FContentBrowserDataFilter& InFilter, FContentBrowserDataCompiledFilter& OutCompiledFilter, FCompileARFilterFunc CreateCompiledFilter);
@@ -364,7 +365,7 @@ protected:
 private:
 	bool IsKnownContentPath(const FName InPackagePath) const;
 
-	static bool GetObjectPathsForCollections(TArrayView<const FCollectionRef> InCollections, const bool bIncludeChildCollections, TArray<FSoftObjectPath>& OutObjectPaths);
+	static bool GetObjectPathsForCollections(TArrayView<const FCollectionRef> InCollections, const bool bIncludeChildCollections, FCollectionEnumerationFunc* GetCollectionObjectPathsFunc, TArray<FSoftObjectPath>& OutObjectPaths);
 
 	FContentBrowserItemData CreateAssetFolderItem(const FName InInternalFolderPath);
 
diff --git a/Engine/Source/Editor/ContentBrowser/Private/AssetContextMenu.cpp b/Engine/Source/Editor/ContentBrowser/Private/AssetContextMenu.cpp
--- a/Engine/Source/Editor/ContentBrowser/Private/AssetContextMenu.cpp
+++ b/Engine/Source/Editor/ContentBrowser/Private/AssetContextMenu.cpp
@@ -458,9 +458,6 @@ void FAssetContextMenu::RegisterContextMenu(const FName MenuName)
 
 void FAssetContextMenu::AddMenuOptions(UToolMenu* InMenu)
 {
-	UContentBrowserDataMenuContext_FileMenu* Context = InMenu->FindContext<UContentBrowserDataMenuContext_FileMenu>();
-	const bool bCanBeModified = !Context || Context->bCanBeModified;
-
 	// Add any type-specific context menu options
 	AddAssetTypeMenuOptions(InMenu);
 
@@ -480,10 +477,7 @@ void FAssetContextMenu::AddMenuOptions(UToolMenu* InMenu)
 	AddReferenceMenuOptions(InMenu);
 
 	// Add collection options
-	if (bCanBeModified)
-	{
-		AddCollectionMenuOptions(InMenu);
-	}
+	AddCollectionMenuOptions(InMenu);
 }
 
 void FAssetContextMenu::SetSelectedItems(TArrayView<const FContentBrowserItem> InSelectedItems)
