[ContentBrowser] Add Menu restyle (via cvar)

diff --git a/Engine/Plugins/Editor/ContentBrowser/ContentBrowserAssetDataSource/Source/ContentBrowserAssetDataSource/Private/ContentBrowserAssetDataSource.cpp b/Engine/Plugins/Editor/ContentBrowser/ContentBrowserAssetDataSource/Source/ContentBrowserAssetDataSource/Private/ContentBrowserAssetDataSource.cpp
--- a/Engine/Plugins/Editor/ContentBrowser/ContentBrowserAssetDataSource/Source/ContentBrowserAssetDataSource/Private/ContentBrowserAssetDataSource.cpp
+++ b/Engine/Plugins/Editor/ContentBrowser/ContentBrowserAssetDataSource/Source/ContentBrowserAssetDataSource/Private/ContentBrowserAssetDataSource.cpp
@@ -11,6 +11,7 @@
 #include "AssetViewUtils.h"
 #include "CollectionManagerModule.h"
 #include "ContentBrowserAssetDataPayload.h"
+#include "ContentBrowserUtils.h"
 #include "ICollectionContainer.h"
 #include "AssetViewUtils.h"
 #include "ContentBrowserItemPath.h"
@@ -48,24 +49,6 @@
 
 #define LOCTEXT_NAMESPACE "ContentBrowserAssetDataSource"
 
-namespace UE::Editor::ContentBrowser
-{
-	static bool IsNewStyleEnabled()
-	{
-		static bool bIsNewStyleEnabled = false;
-		UE_CALL_ONCE([&]()
-		{
-			if (const IConsoleVariable* CVar = IConsoleManager::Get().FindConsoleVariable(TEXT("ContentBrowser.EnableNewStyle")))
-			{
-				ensureAlwaysMsgf(!EnumHasAnyFlags(CVar->GetFlags(), ECVF_Default), TEXT("The CVar should have already been set from commandline, @see: UnrealEdGlobals.cpp, UE::Editor::ContentBrowser::EnableContentBrowserNewStyleCVarRegistration."));
-				bIsNewStyleEnabled = CVar->GetBool();
-			}
-		});
-
-		return bIsNewStyleEnabled;
-	}
-}
-
 namespace AssetDataSource
 {
 	bool bAllowInternalParallelism = true;
diff --git a/Engine/Plugins/Editor/ContentBrowser/ContentBrowserAssetDataSource/Source/ContentBrowserAssetDataSource/Private/ContentBrowserUtils.h b/Engine/Plugins/Editor/ContentBrowser/ContentBrowserAssetDataSource/Source/ContentBrowserAssetDataSource/Private/ContentBrowserUtils.h
new file mode 100644
--- /dev/null
+++ b/Engine/Plugins/Editor/ContentBrowser/ContentBrowserAssetDataSource/Source/ContentBrowserAssetDataSource/Private/ContentBrowserUtils.h
@@ -0,0 +1,26 @@
+// Copyright Epic Games, Inc. All Rights Reserved.
+
+#pragma once
+
+#include "HAL/IConsoleManager.h"
+#include "Misc/AssertionMacros.h"
+#include "Misc/CoreMiscDefines.h"
+#include "Misc/EnumClassFlags.h"
+
+namespace UE::Editor::ContentBrowser
+{
+	static bool IsNewStyleEnabled()
+	{
+		static bool bIsNewStyleEnabled = false;
+		UE_CALL_ONCE([&]()
+		{
+			if (const IConsoleVariable* CVar = IConsoleManager::Get().FindConsoleVariable(TEXT("ContentBrowser.EnableNewStyle")))
+			{
+				ensureAlwaysMsgf(!EnumHasAnyFlags(CVar->GetFlags(), ECVF_Default), TEXT("The CVar should have already been set from commandline, @see: UnrealEdGlobals.cpp, UE::Editor::ContentBrowser::EnableContentBrowserNewStyleCVarRegistration."));
+				bIsNewStyleEnabled = CVar->GetBool();
+			}
+		});
+
+		return bIsNewStyleEnabled;
+	}
+}
diff --git a/Engine/Plugins/Editor/ContentBrowser/ContentBrowserAssetDataSource/Source/ContentBrowserAssetDataSource/Private/NewAssetContextMenu.cpp b/Engine/Plugins/Editor/ContentBrowser/ContentBrowserAssetDataSource/Source/ContentBrowserAssetDataSource/Private/NewAssetContextMenu.cpp
--- a/Engine/Plugins/Editor/ContentBrowser/ContentBrowserAssetDataSource/Source/ContentBrowserAssetDataSource/Private/NewAssetContextMenu.cpp
+++ b/Engine/Plugins/Editor/ContentBrowser/ContentBrowserAssetDataSource/Source/ContentBrowserAssetDataSource/Private/NewAssetContextMenu.cpp
@@ -1,6 +1,7 @@
 // Copyright Epic Games, Inc. All Rights Reserved.
 
 #include "NewAssetContextMenu.h"
+
 #include "ToolMenu.h"
 #include "Widgets/SBoxPanel.h"
 #include "ToolMenuEntry.h"
@@ -11,12 +12,12 @@
 #include "Widgets/Text/STextBlock.h"
 #include "Factories/Factory.h"
 #include "IAssetTools.h"
-#include "IAssetTypeActions.h"
 #include "AssetToolsModule.h"
 #include "Widgets/SToolTip.h"
 #include "IDocumentation.h"
-#include "ClassIconFinder.h"
-#include "AssetToolsModule.h"
+#include "ContentBrowserUtils.h"
+#include "Styling/StyleColors.h"
+#include "Widgets/SAssetMenuIcon.h"
 
 #define LOCTEXT_NAMESPACE "ContentBrowser"
 
@@ -121,11 +122,11 @@ class SFactoryMenuEntry : public SCompoundWidget
 public:
 
 	SLATE_BEGIN_ARGS(SFactoryMenuEntry)
-		: _Width(32)
-		, _Height(32)
+		: _IconContainerSize(32, 32)
+		, _IconSize(28, 28)
 	{}
-		SLATE_ARGUMENT(uint32, Width)
-		SLATE_ARGUMENT(uint32, Height)
+		SLATE_ARGUMENT(FVector2D, IconContainerSize)
+		SLATE_ARGUMENT(FVector2D, IconSize)
 	SLATE_END_ARGS()
 
 	/**
@@ -134,79 +135,57 @@ public:
 	 * @param	InArgs				Declaration used by the SNew() macro to construct this widget
 	 * @param	Factory				The factory this menu entry represents
 	 */
-	void Construct(const FArguments& InArgs, UFactory* Factory)
+	void Construct(const FArguments& InArgs, const UFactory* Factory)
 	{
-		const FName ClassThumbnailBrushOverride = Factory->GetNewAssetThumbnailOverride();
-		const FSlateBrush* ClassThumbnail = nullptr;
-		if (ClassThumbnailBrushOverride.IsNone())
-		{
-			ClassThumbnail = FClassIconFinder::FindThumbnailForClass(Factory->GetSupportedClass());
-		}
-		else
-		{
-			// Instead of getting the override thumbnail directly from the editor style here get it from the
-			// ClassIconFinder since it may have additional styles registered which can be searched by passing
-			// it as a default with no class to search for.
-			ClassThumbnail = FClassIconFinder::FindThumbnailForClass(nullptr, ClassThumbnailBrushOverride);
-		}
+		const TSharedPtr<SWidget> IconContainer =
+			SNew(SAssetMenuIcon,
+				Factory->GetSupportedClass(),
+				UE::Editor::ContentBrowser::IsNewStyleEnabled() ? Factory->GetNewAssetIconOverride() : Factory->GetNewAssetThumbnailOverride())
+			.IconContainerSize(InArgs._IconContainerSize)
+			.IconSize(InArgs._IconSize);
 
-		FAssetToolsModule& AssetToolsModule = FAssetToolsModule::GetModule();
-		TWeakPtr<IAssetTypeActions> AssetTypeActions = AssetToolsModule.Get().GetAssetTypeActionsForClass(Factory->GetSupportedClass());
+		static const FMargin IconSlotPadding = UE::Editor::ContentBrowser::IsNewStyleEnabled()
+			? FMargin(2, 0, 3, 0)	// Consistent with SMenuEntryBlock::BuildMenuEntryWidget, but accounts for icon size that is larger than the default
+			: FMargin(0, 0, 0, 1);
 
-		FLinearColor AssetColor = FLinearColor::White;
-		if (AssetTypeActions.IsValid())
-		{
-			AssetColor = AssetTypeActions.Pin()->GetTypeColor();
-		}
+		static const FMargin LabelSlotPadding = UE::Editor::ContentBrowser::IsNewStyleEnabled()
+			? FMargin(4, 0, 6, 0)	// Consistent with SMenuEntryBlock::BuildMenuEntryWidget
+			: FMargin(4, 0, 4, 0);
+
+		// The vertical padding between each menu entry
+		static constexpr float VerticalEntryPadding = 4.0f;
+
+		// Represents the default icon size in a menu entry, not the one used in this widget
+		static constexpr float DefaultMenuIconSize = 14.0f;
+
+		// Adjust the vertical padding to match the default menu entry spacing as much as possible, while keeping a desired padding minimum of 1px
+		static const float VerticalPaddingAdjustment = ((InArgs._IconContainerSize.Y - DefaultMenuIconSize) / 2.0f) - VerticalEntryPadding;
+
+		static const FMargin ChildSlotPadding = UE::Editor::ContentBrowser::IsNewStyleEnabled()
+			? FMargin(0, -VerticalPaddingAdjustment, 0, -VerticalPaddingAdjustment) // Offset to align with regular menu entries
+			: FMargin(0);
 
 		ChildSlot
-			[
-				SNew(SHorizontalBox)
+		.Padding(ChildSlotPadding)
+		[
+			SNew(SHorizontalBox)
 
-				+SHorizontalBox::Slot()
-				.Padding(4, 0, 0, 0)
-				.VAlign(VAlign_Center)
-				.AutoWidth()
-				[
-					SNew(SOverlay)
-
-					+SOverlay::Slot()
-					[
-						SNew(SBox)
-						.WidthOverride(InArgs._Width + 4)
-						.HeightOverride(InArgs._Height + 4)
-						[
-							SNew(SBorder)
-							.BorderImage(FAppStyle::GetBrush("AssetThumbnail.AssetBackground"))
-							.BorderBackgroundColor(AssetColor.CopyWithNewOpacity(0.3f))
-							.Padding(2.0f)
-							.VAlign(VAlign_Center)
-							.HAlign(HAlign_Center)
-							[
-								SNew(SImage)
-								.Image(ClassThumbnail)
-							]
-						]
-					]
-
-					+SOverlay::Slot()
-					.HAlign(HAlign_Fill)
-					.VAlign(VAlign_Bottom)
-					[
-						SNew(SBorder)
-						.BorderImage(FAppStyle::GetBrush("WhiteBrush"))
-						.BorderBackgroundColor(AssetColor)
-						.Padding(FMargin(0, FMath::Max(FMath::CeilToFloat(InArgs._Width * 0.025f), 3.0f), 0, 0))
-					]
-				]
+			+ SHorizontalBox::Slot()
+			.Padding(IconSlotPadding)
+			.VAlign(VAlign_Center)
+			.AutoWidth()
+			[
+				IconContainer.ToSharedRef()
+			]
 
-			+SHorizontalBox::Slot()
+			+ SHorizontalBox::Slot()
 			.VAlign(VAlign_Center)
-			.Padding(4, 0, 4, 0)
+			.Padding(LabelSlotPadding)
 			[
 				SNew(SVerticalBox)
-				+SVerticalBox::Slot()
-				.Padding(0, 0, 0, 1)
+
+				+ SVerticalBox::Slot()
+				.Padding(UE::Editor::ContentBrowser::IsNewStyleEnabled() ? FMargin(0) : FMargin(0, 0, 0, 1))
 				.AutoHeight()
 				[
 					SNew(STextBlock)
@@ -282,11 +261,29 @@ void FNewAssetContextMenu::MakeContextMenu(
 
 	if (InOnNewAssetRequested.IsBound())
 	{
+		// Add Create Section
+		if (UE::Editor::ContentBrowser::IsNewStyleEnabled())
+		{
+			FToolMenuSection& Section = Menu->AddSection("ContentBrowserNewAsset", LOCTEXT("CreateAssetsMenuHeading", "Create"));
+		}
+
 		// Add Basic Asset
 		{
 			QUICK_SCOPE_CYCLE_COUNTER(STAT_ContentBrowser_BasicSection);
-		
-			FToolMenuSection& Section = Menu->AddSection("ContentBrowserNewBasicAsset", LOCTEXT("CreateBasicAssetsMenuHeading", "Create Basic Asset"));
+
+			const FText CreateBasicAssetSectionLabel =
+				UE::Editor::ContentBrowser::IsNewStyleEnabled()
+				? FText::GetEmpty()
+				: LOCTEXT("CreateBasicAssetsMenuHeading", "Create Basic Asset");
+
+			FToolMenuSection& Section = Menu->AddSection("ContentBrowserNewBasicAsset", CreateBasicAssetSectionLabel);
+
+			if (UE::Editor::ContentBrowser::IsNewStyleEnabled())
+			{
+				// When a section label is empty, it has no visual representation, so manually insert a separator
+				Section.AddSeparator(NAME_None);
+			}
+
 			CreateNewAssetMenuCategory(
 				Menu,
 				"ContentBrowserNewBasicAsset",
@@ -295,16 +292,29 @@ void FNewAssetContextMenu::MakeContextMenu(
 				InOnNewAssetRequested,
 				CanExecuteAssetActionsDelegate
 				);
+
+
 		}
 
 		// Add Advanced Asset
 		{
 			QUICK_SCOPE_CYCLE_COUNTER(STAT_ContentBrowser_AdvancedSection);
-		
-			FToolMenuSection& Section = Menu->AddSection("ContentBrowserNewAdvancedAsset", LOCTEXT("CreateAdvancedAssetsMenuHeading", "Create Advanced Asset"));
+
+			const FText CreateAdvancedAssetSectionLabel =
+				UE::Editor::ContentBrowser::IsNewStyleEnabled()
+				? FText::GetEmpty()
+				: LOCTEXT("CreateAdvancedAssetsMenuHeading", "Create Advanced Asset");
+
+			FToolMenuSection& Section = Menu->AddSection("ContentBrowserNewAdvancedAsset", CreateAdvancedAssetSectionLabel);
+
+			if (UE::Editor::ContentBrowser::IsNewStyleEnabled())
+			{
+				// When a section label is empty, it has no visual representation, so manually insert a separator
+				Section.AddSeparator(NAME_None);
+			}
 
 			TArray<FAdvancedAssetCategory> AdvancedAssetCategories;
-			AssetToolsModule.Get().GetAllAdvancedAssetCategories(/*out*/ AdvancedAssetCategories);
+			AssetToolsModule.Get().GetAllAdvancedAssetCategories(AdvancedAssetCategories);
 			AdvancedAssetCategories.Sort([](const FAdvancedAssetCategory& A, const FAdvancedAssetCategory& B) {
 				return (A.CategoryName.CompareToCaseIgnored(B.CategoryName) < 0);
 			});
@@ -318,24 +328,27 @@ void FNewAssetContextMenu::MakeContextMenu(
 				TArray<FFactoryItem> Factories = FindFactoriesInCategory(NewAssetFactories, AdvancedAssetCategory.CategoryType, FindFirstOnly);
 				if (Factories.Num() > 0)
 				{
-					Section.AddSubMenu(
-						NAME_None,
-						AdvancedAssetCategory.CategoryName,
-						FText::GetEmpty(),
-						FNewToolMenuDelegate::CreateStatic(
-							&FNewAssetContextMenu::CreateNewAssetMenuCategory,
-							FName("Section"),
-							AdvancedAssetCategory.CategoryType,
-							FirstSelectedPath,
-							InOnNewAssetRequested,
-							FCanExecuteAction() // We handle this at this level, rather than at the sub-menu item level
-						),
-						FUIAction(
-							FExecuteAction(),
-							CanExecuteAssetActionsDelegate
-						),
-						EUserInterfaceActionType::Button
-					);
+					FToolMenuEntry& SubMenuEntry =
+							Section.AddSubMenu(
+							NAME_None,
+							AdvancedAssetCategory.CategoryName,
+							FText::GetEmpty(),
+							FNewToolMenuDelegate::CreateStatic(
+								&FNewAssetContextMenu::CreateNewAssetMenuCategory,
+								FName("Section"),
+								AdvancedAssetCategory.CategoryType,
+								FirstSelectedPath,
+								InOnNewAssetRequested,
+								FCanExecuteAction() // We handle this at this level, rather than at the sub-menu item level
+							),
+							FUIAction(
+								FExecuteAction(),
+								CanExecuteAssetActionsDelegate
+							),
+							EUserInterfaceActionType::Button
+						);
+
+					SubMenuEntry.SubMenuData.Style.StyleName = "ContentBrowser.AddNewMenu";
 				}
 			}
 		}
@@ -376,7 +389,7 @@ void FNewAssetContextMenu::CreateNewAssetMenuCategory(UToolMenu* Menu, FName Sec
 		SubMenu->Factories.Add(Item);
 	}
 	ParentMenuData->SortSubMenus();
-	
+
 	CreateNewAssetMenus(Menu, SectionName, ParentMenuData, InPath, InOnNewAssetRequested, InCanExecuteAction);
 }
 
@@ -396,15 +409,39 @@ void FNewAssetContextMenu::CreateNewAssetMenus(UToolMenu* Menu, FName SectionNam
 			AssetTypeName = SupportedClass->GetFName();
 		}
 
-		Section.AddEntry(FToolMenuEntry::InitMenuEntry(
-			NAME_None,
-			FUIAction(
-				FExecuteAction::CreateStatic(&FNewAssetContextMenu::ExecuteNewAsset, InOnNewAssetRequested, InPath, WeakFactoryClass),
-				InCanExecuteAction
-			),
+		if (UE::Editor::ContentBrowser::IsNewStyleEnabled())
+		{
+			static constexpr uint32 IconContainerSize = 24;
+			static constexpr uint32 IconSize = 16;
+
+			FToolMenuEntry& Entry = Section.AddEntry(
+				FToolMenuEntry::InitMenuEntry(
+				NAME_None,
+				FUIAction(
+					FExecuteAction::CreateStatic(&FNewAssetContextMenu::ExecuteNewAsset, InOnNewAssetRequested, InPath, WeakFactoryClass),
+					InCanExecuteAction
+				),
 			SNew(SFactoryMenuEntry, &FactoryItem.Factory)
-				.AddMetaData<FTagMetaData>(FTagMetaData(AssetTypeName))));
+				.IconContainerSize(FVector2D(IconContainerSize, IconContainerSize))
+				.IconSize(FVector2D(IconSize, IconSize))
+				.AddMetaData<FTagMetaData>(FTagMetaData(AssetTypeName)))
+			);
 
+			Entry.SubMenuData.Style.StyleName = "ContentBrowser.AddNewMenu";
+		}
+		else
+		{
+			Section.AddEntry(
+				FToolMenuEntry::InitMenuEntry(
+				NAME_None,
+				FUIAction(
+					FExecuteAction::CreateStatic(&FNewAssetContextMenu::ExecuteNewAsset, InOnNewAssetRequested, InPath, WeakFactoryClass),
+					InCanExecuteAction
+				),
+			SNew(SFactoryMenuEntry, &FactoryItem.Factory)
+				.AddMetaData<FTagMetaData>(FTagMetaData(AssetTypeName)))
+			);
+		}
 	}
 
 	if (SubMenuData->Children.Num() == 0)
@@ -425,7 +462,7 @@ void FNewAssetContextMenu::CreateNewAssetMenus(UToolMenu* Menu, FName SectionNam
 	{
 		check(ChildMenuData.IsValid());
 
-		Section.AddSubMenu(
+		FToolMenuEntry& Entry = Section.AddSubMenu(
 			NAME_None,
 			ChildMenuData->Name,
 			FText::GetEmpty(),
@@ -443,6 +480,11 @@ void FNewAssetContextMenu::CreateNewAssetMenus(UToolMenu* Menu, FName SectionNam
 			),
 			EUserInterfaceActionType::Button
 		);
+
+		if (UE::Editor::ContentBrowser::IsNewStyleEnabled())
+		{
+			Entry.SubMenuData.Style.StyleName = "ContentBrowser.AddNewMenu";
+		}
 	}
 }
 
diff --git a/Engine/Plugins/Editor/EngineAssetDefinitions/Source/Private/Animation/AssetDefinition_AnimBlueprint.h b/Engine/Plugins/Editor/EngineAssetDefinitions/Source/Private/Animation/AssetDefinition_AnimBlueprint.h
--- a/Engine/Plugins/Editor/EngineAssetDefinitions/Source/Private/Animation/AssetDefinition_AnimBlueprint.h
+++ b/Engine/Plugins/Editor/EngineAssetDefinitions/Source/Private/Animation/AssetDefinition_AnimBlueprint.h
@@ -24,7 +24,7 @@ public:
 	virtual TSoftClassPtr<UObject> GetAssetClass() const override { return UAnimBlueprint::StaticClass(); }
 	virtual TConstArrayView<FAssetCategoryPath> GetAssetCategories() const override
 	{
-		static const auto Categories = { EAssetCategoryPaths::Animation };
+		static const auto Categories = { EAssetCategoryPaths::Basic, EAssetCategoryPaths::Animation };
 		return Categories;
 	}
 	
diff --git a/Engine/Plugins/MovieScene/LevelSequenceEditor/Source/LevelSequenceEditor/Private/AssetTools/AssetDefinition_LevelSequence.h b/Engine/Plugins/MovieScene/LevelSequenceEditor/Source/LevelSequenceEditor/Private/AssetTools/AssetDefinition_LevelSequence.h
--- a/Engine/Plugins/MovieScene/LevelSequenceEditor/Source/LevelSequenceEditor/Private/AssetTools/AssetDefinition_LevelSequence.h
+++ b/Engine/Plugins/MovieScene/LevelSequenceEditor/Source/LevelSequenceEditor/Private/AssetTools/AssetDefinition_LevelSequence.h
@@ -19,7 +19,7 @@ public:
 	virtual TSoftClassPtr<UObject> GetAssetClass() const override { return ULevelSequence::StaticClass(); }
 	virtual TConstArrayView<FAssetCategoryPath> GetAssetCategories() const override
 	{
-		static const auto Categories = { EAssetCategoryPaths::Cinematics };
+		static const auto Categories = { EAssetCategoryPaths::Basic, EAssetCategoryPaths::Cinematics };
 		return Categories;
 	}
 	virtual FAssetOpenSupport GetAssetOpenSupport(const FAssetOpenSupportArgs& OpenSupportArgs) const override;
diff --git a/Engine/Source/Editor/ContentBrowser/Private/ContentBrowserStyle.cpp b/Engine/Source/Editor/ContentBrowser/Private/ContentBrowserStyle.cpp
--- a/Engine/Source/Editor/ContentBrowser/Private/ContentBrowserStyle.cpp
+++ b/Engine/Source/Editor/ContentBrowser/Private/ContentBrowserStyle.cpp
@@ -355,6 +355,37 @@ namespace UE::ContentBrowser::Private
 		// Separate style for buttons prior to the search box to get around SToolBarComboButtonBlock referencing the wrong style
 		Set("ContentBrowser.ToolBar.Buttons", ContentBrowserToolBarStyle);
 
+		Set("ContentBrowser.AddNewMenu.Separator", new FSlateColorBrush(FStyleColors::White25));
+		Set("ContentBrowser.AddNewMenu.Separator.Padding", FAppStyle::Get().GetMargin("Menu.Separator.Padding"));
+
+		Set("ContentBrowser.AddNewMenu.Label", FAppStyle::Get().GetWidgetStyle<FTextBlockStyle>("Menu.Label"));
+		Set("ContentBrowser.AddNewMenu.Heading", FAppStyle::Get().GetWidgetStyle<FTextBlockStyle>("Menu.Heading"));
+		Set("ContentBrowser.AddNewMenu.Heading.Padding", FAppStyle::Get().GetMargin("Menu.Heading.Padding"));
+		Set("ContentBrowser.AddNewMenu.Keybinding", FAppStyle::Get().GetWidgetStyle<FTextBlockStyle>("Menu.Keybinding"));
+
+		Set("ContentBrowser.AddNewMenu.SubMenuIndicator", new CORE_IMAGE_BRUSH_SVG("Starship/Common/chevron-right", CoreStyleConstants::Icon16x16, FStyleColors::Foreground));
+
+		Set("ContentBrowser.AddNewMenu.Button", FAppStyle::Get().GetWidgetStyle<FButtonStyle>("Menu.Button"));
+		Set("ContentBrowser.AddNewMenu.CheckBox", FAppStyle::Get().GetWidgetStyle<FCheckBoxStyle>("Menu.CheckBox"));
+
+		FMargin MenuBlockPadding = FAppStyle::Get().GetMargin("Menu.Block.Padding");
+		FMargin IndentedMenuBlockPadding = FAppStyle::Get().GetMargin("Menu.Block.IndentedPadding");
+
+		if (UE::Editor::ContentBrowser::IsNewStyleEnabled())
+		{
+			MenuBlockPadding.Left = 0.0f;
+			IndentedMenuBlockPadding.Left = 0.0f;
+		}
+
+		Set("ContentBrowser.AddNewMenu.Block.Padding", MenuBlockPadding);
+		Set("ContentBrowser.AddNewMenu.Block.IndentedPadding", IndentedMenuBlockPadding);
+
+		MenuBlockPadding.Left = -100.0f;
+		IndentedMenuBlockPadding.Left = -150.0f;
+
+		Set("ContentBrowser.AssetViewOptions.Block.Padding", MenuBlockPadding);
+		Set("ContentBrowser.AssetViewOptions.Block.IndentedPadding", IndentedMenuBlockPadding);
+
 		FSlateStyleRegistry::RegisterSlateStyle(*this);
 	}
 
diff --git a/Engine/Source/Editor/ContentBrowser/Private/Menus/ContentBrowserMenus.cpp b/Engine/Source/Editor/ContentBrowser/Private/Menus/ContentBrowserMenus.cpp
--- a/Engine/Source/Editor/ContentBrowser/Private/Menus/ContentBrowserMenus.cpp
+++ b/Engine/Source/Editor/ContentBrowser/Private/Menus/ContentBrowserMenus.cpp
@@ -126,6 +126,8 @@ FDelayedAutoRegisterHelper SContentBrowser::AddNewContextMenuRegistration(
 
         static const FName AddNewContextMenuName("ContentBrowser.AddNewContextMenu");
         UToolMenu* Menu = UToolMenus::Get()->RegisterMenu(AddNewContextMenuName);
+        Menu->SetStyleSet(&UE::ContentBrowser::Private::FContentBrowserStyle::Get());
+		Menu->StyleName = "ContentBrowser.AddNewMenu";
 
         Menu->AddDynamicSection("DynamicSection_Common", FNewToolMenuDelegate::CreateLambda([](UToolMenu* InMenu)
         {
diff --git a/Engine/Source/Editor/ContentBrowser/Private/NewAssetOrClassContextMenu.cpp b/Engine/Source/Editor/ContentBrowser/Private/NewAssetOrClassContextMenu.cpp
--- a/Engine/Source/Editor/ContentBrowser/Private/NewAssetOrClassContextMenu.cpp
+++ b/Engine/Source/Editor/ContentBrowser/Private/NewAssetOrClassContextMenu.cpp
@@ -42,7 +42,13 @@ void FNewAssetOrClassContextMenu::MakeContextMenu(
 	const FCanExecuteAction CanExecuteFolderActionsDelegate = FCanExecuteAction::CreateLambda(CanExecuteFolderActions);
 
 	// Get Content
-	FToolMenuSection& GetContentSection = Menu->AddSection("ContentBrowserGetContent", LOCTEXT("GetContentMenuHeading", "Get Content"));
+
+	const FText GetContentSectionLabel =
+		UE::Editor::ContentBrowser::IsNewStyleEnabled()
+		? LOCTEXT("GetContentMenuHeading_NewStyle", "Get")
+		: LOCTEXT("GetContentMenuHeading", "Get Content");
+
+	FToolMenuSection& GetContentSection = Menu->AddSection("ContentBrowserGetContent", GetContentSectionLabel);
 	UContentBrowserDataMenuContext_AddNewMenu* AddNewMenuContext = Menu->FindContext<UContentBrowserDataMenuContext_AddNewMenu>();
 
 	if ( InOnGetContentRequested.IsBound() )
@@ -76,7 +82,18 @@ void FNewAssetOrClassContextMenu::MakeContextMenu(
 	if(InOnNewFolderRequested.IsBound() && bDisplayFolders)
 	{
 		{
-			FToolMenuSection& Section = Menu->AddSection("ContentBrowserNewFolder", LOCTEXT("FolderMenuHeading", "Folder") );
+			FToolMenuSection* Section = nullptr;
+			if (UE::Editor::ContentBrowser::IsNewStyleEnabled())
+			{
+				Section = &Menu->FindOrAddSection("ContentBrowserNewAsset");
+			}
+			else
+			{
+				Section = &Menu->AddSection("ContentBrowserNewFolder", LOCTEXT("FolderMenuHeading", "Folder") );
+			}
+
+			check(Section);
+
 			FText NewFolderToolTip;
 			if(bHasSinglePathSelected)
 			{
@@ -94,11 +111,11 @@ void FNewAssetOrClassContextMenu::MakeContextMenu(
 				NewFolderToolTip = LOCTEXT("NewFolderTooltip_InvalidNumberOfPaths", "Can only create folders when there is a single path selected.");
 			}
 
-			Section.AddMenuEntry(
+			Section->AddMenuEntry(
 				"NewFolder",
 				LOCTEXT("NewFolderLabel", "New Folder"),
 				NewFolderToolTip,
-				FSlateIcon(UE::ContentBrowser::Private::FContentBrowserStyle::Get().GetStyleSetName(), "ContentBrowser.NewFolderIcon"),
+				FSlateIcon(FAppStyle::GetAppStyleSetName(), "ContentBrowser.NewFolderIcon"),
 				FUIAction(
 					FExecuteAction::CreateStatic(&FNewAssetOrClassContextMenu::ExecuteNewFolder, FirstSelectedPath, InOnNewFolderRequested),
 					CanExecuteFolderActionsDelegate
diff --git a/Engine/Source/Editor/ContentBrowser/Private/Widgets/SAssetMenuIcon.cpp b/Engine/Source/Editor/ContentBrowser/Private/Widgets/SAssetMenuIcon.cpp
new file mode 100644
--- /dev/null
+++ b/Engine/Source/Editor/ContentBrowser/Private/Widgets/SAssetMenuIcon.cpp
@@ -0,0 +1,128 @@
+// Copyright Epic Games, Inc. All Rights Reserved.
+
+#include "Widgets/SAssetMenuIcon.h"
+
+#include "AssetToolsModule.h"
+#include "ClassIconFinder.h"
+#include "ContentBrowserStyle.h"
+#include "Misc/CommandLine.h"
+#include "SlateOptMacros.h"
+
+void SAssetMenuIcon::Construct(const FArguments& InArgs, const UClass* InAssetClass, const FName InIconOverride)
+{
+	check(InAssetClass);
+
+	FAssetToolsModule& AssetToolsModule = FAssetToolsModule::GetModule();
+	TWeakPtr<IAssetTypeActions> AssetTypeActions = AssetToolsModule.Get().GetAssetTypeActionsForClass(InAssetClass);
+
+	FLinearColor AssetColor = FLinearColor::White;
+	if (AssetTypeActions.IsValid())
+	{
+		AssetColor = AssetTypeActions.Pin()->GetTypeColor();
+	}
+
+	if (UE::Editor::ContentBrowser::IsNewStyleEnabled())
+	{
+		static const FVector2D ChipSize = FVector2D(6, 6);
+
+		const FName ClassIconBrushOverride = InIconOverride;
+		const FSlateBrush* ClassIcon = nullptr;
+		if (ClassIconBrushOverride.IsNone())
+		{
+			ClassIcon = FSlateIconFinder::FindIconBrushForClass(InAssetClass);
+		}
+		else
+		{
+			// Instead of getting the override icon directly from the editor style here get it from the
+			// FSlateIconFinder since it may have additional styles registered which can be searched by passing
+			// it as a default with no class to search for.
+			ClassIcon = FSlateIconFinder::FindIconBrushForClass(nullptr, ClassIconBrushOverride);
+		}
+
+		ChildSlot
+		[
+			SNew(SBox)
+			.WidthOverride(static_cast<float>(InArgs._IconContainerSize.X))
+			.HeightOverride(static_cast<float>(InArgs._IconContainerSize.Y))
+			[
+				SNew(SOverlay)
+
+				+ SOverlay::Slot()
+				[
+					SNew(SBorder)
+					.BorderImage(FAppStyle::GetBrush("WhiteBrush"))
+					.BorderBackgroundColor(FSlateColor(FStyleColors::Background))
+					.Padding(2.0f)
+					.VAlign(VAlign_Center)
+					.HAlign(HAlign_Center)
+					[
+						SNew(SImage)
+						.DesiredSizeOverride(FVector2D(InArgs._IconSize.X, InArgs._IconSize.X))
+						.ColorAndOpacity(FSlateColor::UseForeground())
+						.Image(ClassIcon)
+					]
+				]
+
+				+ SOverlay::Slot()
+				.HAlign(HAlign_Right)
+				.VAlign(VAlign_Top)
+				.Padding(0)
+				[
+					SNew(SImage)
+					.DesiredSizeOverride(ChipSize)
+					.Image(FAppStyle::GetBrush("AssetThumbnail.AssetTileViewWhiteChipSharp"))
+					.ColorAndOpacity(AssetColor)
+				]
+			]
+		];
+	}
+	else
+	{
+		const FName ClassThumbnailBrushOverride = InIconOverride;
+		const FSlateBrush* ClassThumbnail = nullptr;
+		if (ClassThumbnailBrushOverride.IsNone())
+		{
+			ClassThumbnail = FClassIconFinder::FindThumbnailForClass(InAssetClass);
+		}
+		else
+		{
+			// Instead of getting the override thumbnail directly from the editor style here get it from the
+			// ClassIconFinder since it may have additional styles registered which can be searched by passing
+			// it as a default with no class to search for.
+			ClassThumbnail = FClassIconFinder::FindThumbnailForClass(nullptr, ClassThumbnailBrushOverride);
+		}
+
+		ChildSlot
+		[
+			SNew(SOverlay)
+
+			+ SOverlay::Slot()
+			[
+				SNew(SBox)
+				.WidthOverride(static_cast<float>(InArgs._IconContainerSize.X))
+				.HeightOverride(static_cast<float>(InArgs._IconContainerSize.Y))
+				[
+					SNew(SBorder)
+					.BorderImage(FAppStyle::GetBrush("AssetThumbnail.AssetBackground"))
+					.Padding(2.0f)
+					.VAlign(VAlign_Center)
+					.HAlign(HAlign_Center)
+					[
+						SNew(SImage)
+						.Image(ClassThumbnail)
+					]
+				]
+			]
+
+			+ SOverlay::Slot()
+			.HAlign(HAlign_Fill)
+			.VAlign(VAlign_Bottom)
+			[
+				SNew(SBorder)
+				.BorderImage(FAppStyle::GetBrush("WhiteBrush"))
+				.BorderBackgroundColor(AssetColor)
+				.Padding(FMargin(0, FMath::Max(FMath::CeilToFloat(InArgs._IconSize.X * 0.025f), 3.0f), 0, 0))
+			]
+		];
+	}
+}
diff --git a/Engine/Source/Editor/ContentBrowser/Public/Widgets/SAssetMenuIcon.h b/Engine/Source/Editor/ContentBrowser/Public/Widgets/SAssetMenuIcon.h
new file mode 100644
--- /dev/null
+++ b/Engine/Source/Editor/ContentBrowser/Public/Widgets/SAssetMenuIcon.h
@@ -0,0 +1,26 @@
+// Copyright Epic Games, Inc. All Rights Reserved.
+
+#pragma once
+
+#include "CoreTypes.h"
+#include "Styling/StyleColors.h"
+#include "Widgets/Images/SImage.h"
+#include "Widgets/Layout/SBox.h"
+#include "Widgets/SCompoundWidget.h"
+#include "Widgets/Text/STextBlock.h"
+
+/** Represents an Asset Type within a menu. */
+class CONTENTBROWSER_API SAssetMenuIcon : public SCompoundWidget
+{
+public:
+	SLATE_BEGIN_ARGS(SAssetMenuIcon)
+		: _IconContainerSize(32.0f, 32.0f)
+		, _IconSize(28.0f, 28.0f)
+	{}
+		SLATE_ARGUMENT(FVector2D, IconContainerSize)
+		SLATE_ARGUMENT(FVector2D, IconSize)
+	SLATE_END_ARGS()
+
+	/** Constructs this widget with InArgs */
+	void Construct(const FArguments& InArgs, const UClass* InAssetClass, const FName InIconOverride = NAME_None);
+};
diff --git a/Engine/Source/Editor/EditorStyle/Private/StarshipStyle.cpp b/Engine/Source/Editor/EditorStyle/Private/StarshipStyle.cpp
--- a/Engine/Source/Editor/EditorStyle/Private/StarshipStyle.cpp
+++ b/Engine/Source/Editor/EditorStyle/Private/StarshipStyle.cpp
@@ -1247,6 +1247,7 @@ void FStarshipEditorStyle::FStyle::SetupGeneralStyles()
 		Set( "AssetThumbnail.HintBackground", new BOX_BRUSH( "Common/TableViewHeader", FMargin(8.0f/32.0f) ) );
 		Set( "AssetThumbnail.AssetTileViewWhiteChipBorder", new IMAGE_BRUSH_SVG("Starship/ContentBrowser/CornerChipOuter2_20", FVector2D(20, 20)));
 		Set( "AssetThumbnail.AssetTileViewWhiteChip", new IMAGE_BRUSH_SVG("Starship/ContentBrowser/CornerChipInner2_20", FVector2D(20, 20)));
+		Set( "AssetThumbnail.AssetTileViewWhiteChipSharp", new IMAGE_BRUSH_SVG("Starship/ContentBrowser/CornerChipSharp_20", FVector2D(20, 20)));
 		Set( "AssetThumbnail.AssetBorder", new FSlateRoundedBoxBrush(FStyleColors::Transparent, FVector4(4.0f, 4.0f, 4.0f, 4.0f), FStyleColors::Secondary, 2.f));
 		FLinearColor OverlayColorAssetStatusOverlay = FStyleColors::Panel.GetSpecifiedColor();
 		OverlayColorAssetStatusOverlay.A = 0.75f;
diff --git a/Engine/Source/Editor/UnrealEd/Classes/Factories/Factory.h b/Engine/Source/Editor/UnrealEd/Classes/Factories/Factory.h
--- a/Engine/Source/Editor/UnrealEd/Classes/Factories/Factory.h
+++ b/Engine/Source/Editor/UnrealEd/Classes/Factories/Factory.h
@@ -131,9 +131,12 @@ public:
 	/** Returns true if this factory should be shown in the New Asset menu (by default calls CanCreateNew). */
 	UNREALED_API virtual bool ShouldShowInNewMenu() const;
 
-	/** Returns an optional override brush name for the new asset menu. If this is not specified, the thumbnail for the supported class will be used. */
+	/** Returns an optional override brush name for the new asset menu. If this is not specified, the (64x64) thumbnail for the supported class will be used. */
 	UNREALED_API virtual FName GetNewAssetThumbnailOverride() const;
 
+	/** Returns an optional override brush name for the new asset menu. If this is not specified, the (16x16) icon for the supported class will be used. */
+	UNREALED_API virtual FName GetNewAssetIconOverride() const;
+
 	/** Returns the name of the factory for menus */
 	UNREALED_API virtual FText GetDisplayName() const;
 
diff --git a/Engine/Source/Editor/UnrealEd/Private/Factories/Factory.cpp b/Engine/Source/Editor/UnrealEd/Private/Factories/Factory.cpp
--- a/Engine/Source/Editor/UnrealEd/Private/Factories/Factory.cpp
+++ b/Engine/Source/Editor/UnrealEd/Private/Factories/Factory.cpp
@@ -228,9 +228,15 @@ bool UFactory::ShouldShowInNewMenu() const
 
 FName UFactory::GetNewAssetThumbnailOverride() const
 {
+	// @note: this typically refers to a 64x64 thumbnail
 	return NAME_None;
 }
 
+FName UFactory::GetNewAssetIconOverride() const
+{
+	// @note: this typically refers to a 16x16 icon
+	return NAME_None;
+}
 
 FText UFactory::GetDisplayName() const
 {
