Include the in-memory text cache in the "Asset Localization -> Show Localization Cache" report

This makes it easier to debug cache invalidation issues, plus see what text will be gathered for assets that skip the on-disk cache (like Blueprints)

diff --git a/Engine/Plugins/Editor/ContentBrowser/ContentBrowserAssetDataSource/Source/ContentBrowserAssetDataSource/Private/AssetFileContextMenu.cpp b/Engine/Plugins/Editor/ContentBrowser/ContentBrowserAssetDataSource/Source/ContentBrowserAssetDataSource/Private/AssetFileContextMenu.cpp
--- a/Engine/Plugins/Editor/ContentBrowser/ContentBrowserAssetDataSource/Source/ContentBrowserAssetDataSource/Private/AssetFileContextMenu.cpp
+++ b/Engine/Plugins/Editor/ContentBrowser/ContentBrowserAssetDataSource/Source/ContentBrowserAssetDataSource/Private/AssetFileContextMenu.cpp
@@ -41,6 +41,7 @@
 #include "ReferencedAssetsUtils.h"
 #include "Internationalization/PackageLocalizationUtil.h"
 #include "Internationalization/TextLocalizationResource.h"
+#include "Serialization/PropertyLocalizationDataGathering.h"
 
 #include "Kismet2/KismetEditorUtilities.h"
 #include "ComponentAssetBroker.h"
@@ -820,7 +821,7 @@ void FAssetFileContextMenu::MakeAssetLocalizationSubMenu(UToolMenu* Menu)
 					LOCTEXT("ShowLocalizationCache", "Show Localization Cache"),
 					LOCTEXT("ShowLocalizationCacheTooltip", "Show the cached list of localized texts stored in the package header."),
 					FSlateIcon(),
-					FUIAction(FExecuteAction::CreateSP(this, &FAssetFileContextMenu::ExecuteShowLocalizationCache, PackageFilename))
+					FUIAction(FExecuteAction::CreateSP(this, &FAssetFileContextMenu::ExecuteShowLocalizationCache, SelectedAssets[0], PackageFilename))
 					);
 			}
 		}
@@ -1783,13 +1784,50 @@ void FAssetFileContextMenu::ExecuteResetLocalizationId()
 #endif // USE_STABLE_LOCALIZATION_KEYS
 }
 
-void FAssetFileContextMenu::ExecuteShowLocalizationCache(const FString InPackageFilename)
+void FAssetFileContextMenu::ExecuteShowLocalizationCache(const FAssetData InAsset, const FString InPackageFilename)
 {
-	FString CachedLocalizationId;
-	TArray<FGatherableTextData> GatherableTextDataArray;
+	TStringBuilder<2048> LocalizationCacheStr;
+	auto AppendLocalizationCache = [&LocalizationCacheStr](const FString& CachedLocalizationId, const TArray<FGatherableTextData>& GatherableTextDataArray)
+	{
+		// Convert the gathered text array into a readable format
+		LocalizationCacheStr.Appendf(TEXT("Package: %s"), *CachedLocalizationId);
+		for (const FGatherableTextData& GatherableTextData : GatherableTextDataArray)
+		{
+			if (LocalizationCacheStr.Len() > 0)
+			{
+				LocalizationCacheStr += TEXTVIEW("\n\n");
+			}
+
+			TStringBuilder<512> KeysStr;
+			TStringBuilder<512> EditorOnlyKeysStr;
+			for (const FTextSourceSiteContext& TextSourceSiteContext : GatherableTextData.SourceSiteContexts)
+			{
+				TStringBuilder<512>* KeysStrPtr = TextSourceSiteContext.IsEditorOnly ? &EditorOnlyKeysStr : &KeysStr;
+				if (KeysStrPtr->Len() > 0)
+				{
+					*KeysStrPtr += TEXTVIEW(", ");
+				}
+				*KeysStrPtr += TextSourceSiteContext.KeyName;
+			}
+
+			LocalizationCacheStr.Appendf(TEXT("Namespace: %s\n"), *GatherableTextData.NamespaceName);
+			if (KeysStr.Len() > 0)
+			{
+				LocalizationCacheStr.Appendf(TEXT("Keys: %s\n"), *KeysStr);
+			}
+			if (EditorOnlyKeysStr.Len() > 0)
+			{
+				LocalizationCacheStr.Appendf(TEXT("Keys (Editor-Only): %s\n"), *EditorOnlyKeysStr);
+			}
+			LocalizationCacheStr.Appendf(TEXT("Source: %s"), *GatherableTextData.SourceData.SourceString);
+		}
+	};
 
 	// Read the localization data from the cache in the package header
 	{
+		FString CachedLocalizationId;
+		TArray<FGatherableTextData> GatherableTextDataArray;
+
 		TUniquePtr<FArchive> FileReader(IFileManager::Get().CreateFileReader(*InPackageFilename));
 		if (FileReader)
 		{
@@ -1810,39 +1848,29 @@ void FAssetFileContextMenu::ExecuteShowLocalizationCache(const FString InPackage
 				}
 			}
 		}
+
+		LocalizationCacheStr += TEXTVIEW("== On-Disk Cache ==\n\n");
+		AppendLocalizationCache(CachedLocalizationId, GatherableTextDataArray);
 	}
 
-	// Convert the gathered text array into a readable format
-	FString LocalizationCacheStr = FString::Printf(TEXT("Package: %s"), *CachedLocalizationId);
-	for (const FGatherableTextData& GatherableTextData : GatherableTextDataArray)
+	// If the asset is loaded, also show the "live" version of the cache (for debugging cache invalidation issues)
 	{
-		if (LocalizationCacheStr.Len() > 0)
-		{
-			LocalizationCacheStr += TEXT("\n\n");
-		}
+		LocalizationCacheStr += TEXTVIEW("\n\n\n\n== In-Memory Cache ==\n\n");
 
-		FString KeysStr;
-		FString EditorOnlyKeysStr;
-		for (const FTextSourceSiteContext& TextSourceSiteContext : GatherableTextData.SourceSiteContexts)
+		if (UObject* Asset = InAsset.FastGetAsset())
 		{
-			FString* KeysStrPtr = TextSourceSiteContext.IsEditorOnly ? &EditorOnlyKeysStr : &KeysStr;
-			if (KeysStrPtr->Len() > 0)
-			{
-				*KeysStrPtr += TEXT(", ");
-			}
-			*KeysStrPtr += TextSourceSiteContext.KeyName;
-		}
+			const FString CachedLocalizationId = TextNamespaceUtil::GetPackageNamespace(Asset);
+			TArray<FGatherableTextData> GatherableTextDataArray;
 
-		LocalizationCacheStr += FString::Printf(TEXT("Namespace: %s\n"), *GatherableTextData.NamespaceName);
-		if (KeysStr.Len() > 0)
-		{
-			LocalizationCacheStr += FString::Printf(TEXT("Keys: %s\n"), *KeysStr);
+			EPropertyLocalizationGathererResultFlags GatherableTextResultFlags = EPropertyLocalizationGathererResultFlags::Empty;
+			FPropertyLocalizationDataGatherer _(GatherableTextDataArray, Asset->GetPackage(), GatherableTextResultFlags);
+
+			AppendLocalizationCache(CachedLocalizationId, GatherableTextDataArray);
 		}
-		if (EditorOnlyKeysStr.Len() > 0)
+		else
 		{
-			LocalizationCacheStr += FString::Printf(TEXT("Keys (Editor-Only): %s\n"), *EditorOnlyKeysStr);
+			LocalizationCacheStr += TEXTVIEW("<Asset Not Loaded>");
 		}
-		LocalizationCacheStr += FString::Printf(TEXT("Source: %s"), *GatherableTextData.SourceData.SourceString);
 	}
 
 	// Generate a message box for the result
diff --git a/Engine/Plugins/Editor/ContentBrowser/ContentBrowserAssetDataSource/Source/ContentBrowserAssetDataSource/Public/AssetFileContextMenu.h b/Engine/Plugins/Editor/ContentBrowser/ContentBrowserAssetDataSource/Source/ContentBrowserAssetDataSource/Public/AssetFileContextMenu.h
--- a/Engine/Plugins/Editor/ContentBrowser/ContentBrowserAssetDataSource/Source/ContentBrowserAssetDataSource/Public/AssetFileContextMenu.h
+++ b/Engine/Plugins/Editor/ContentBrowser/ContentBrowserAssetDataSource/Source/ContentBrowserAssetDataSource/Public/AssetFileContextMenu.h
@@ -160,7 +160,7 @@ private:
 	void ExecuteResetLocalizationId();
 
 	/** Handler for showing the cached list of localized texts stored in the package header */
-	void ExecuteShowLocalizationCache(const FString InPackageFilename);
+	void ExecuteShowLocalizationCache(const FAssetData InAsset, const FString InPackageFilename);
 
 	/** Handler for Export */
 	void ExecuteExport();
