Fixed an issue where you couldn't rename a folder to a folder name that is always hidden when empty because the folder exists even though the user can't see it.

diff --git a/Engine/Plugins/Editor/ContentBrowser/ContentBrowserAssetDataSource/Source/ContentBrowserAssetDataSource/Private/ContentBrowserAssetDataCore.cpp b/Engine/Plugins/Editor/ContentBrowser/ContentBrowserAssetDataSource/Source/ContentBrowserAssetDataSource/Private/ContentBrowserAssetDataCore.cpp
--- a/Engine/Plugins/Editor/ContentBrowser/ContentBrowserAssetDataSource/Source/ContentBrowserAssetDataSource/Private/ContentBrowserAssetDataCore.cpp
+++ b/Engine/Plugins/Editor/ContentBrowser/ContentBrowserAssetDataSource/Source/ContentBrowserAssetDataSource/Private/ContentBrowserAssetDataCore.cpp
@@ -959,21 +959,33 @@ bool DeleteUnsupportedAssetFileItems(TArrayView<const TSharedRef<const FContentB
 }
 
 bool CanRenameItem(IAssetTools* InAssetTools, const UContentBrowserDataSource* InOwnerDataSource, const FContentBrowserItemData& InItem, const FString* InNewName, FText* OutErrorMsg)
+{
+	constexpr bool bCheckUniqueName = true;
+	return CanRenameItem(InAssetTools, InOwnerDataSource, InItem, bCheckUniqueName, InNewName, OutErrorMsg);
+}
+
+bool CanRenameItem(IAssetTools* InAssetTools, const UContentBrowserDataSource* InOwnerDataSource, const FContentBrowserItemData& InItem, const bool bCheckUniqueName, const FString* InNewName, FText* OutErrorMsg)
 {
 	if (TSharedPtr<const FContentBrowserAssetFolderItemDataPayload> FolderPayload = GetAssetFolderItemPayload(InOwnerDataSource, InItem))
 	{
-		return CanRenameAssetFolderItem(InAssetTools, *FolderPayload, InNewName, OutErrorMsg);
+		return CanRenameAssetFolderItem(InAssetTools, *FolderPayload, bCheckUniqueName, InNewName, OutErrorMsg);
 	}
 
 	if (TSharedPtr<const FContentBrowserAssetFileItemDataPayload> AssetPayload = GetAssetFileItemPayload(InOwnerDataSource, InItem))
 	{
-		return CanRenameAssetFileItem(InAssetTools, *AssetPayload, InNewName, InItem.IsTemporary(), OutErrorMsg);
+		return CanRenameAssetFileItem(InAssetTools, *AssetPayload, bCheckUniqueName, InNewName, InItem.IsTemporary(), OutErrorMsg);
 	}
 
 	return false;
 }
 
 bool CanRenameAssetFolderItem(IAssetTools* InAssetTools, const FContentBrowserAssetFolderItemDataPayload& InFolderPayload, const FString* InNewName, FText* OutErrorMsg)
+{
+	constexpr bool bCheckUniqueName = true;
+	return CanRenameAssetFolderItem(InAssetTools, InFolderPayload, bCheckUniqueName, InNewName, OutErrorMsg);
+}
+
+bool CanRenameAssetFolderItem(IAssetTools* InAssetTools, const FContentBrowserAssetFolderItemDataPayload& InFolderPayload, const bool bCheckUniqueName, const FString* InNewName, FText* OutErrorMsg)
 {
 	if (!CanModifyAssetFolderItem(InAssetTools, InFolderPayload, OutErrorMsg))
 	{
@@ -991,7 +1003,8 @@ bool CanRenameAssetFolderItem(IAssetTools* InAssetTools, const FContentBrowserAs
 		const FString FolderPath = FPaths::GetPath(InFolderPayload.GetInternalPath().ToString());
 
 		FText ValidationErrorMsg;
-		if (!AssetViewUtils::IsValidFolderPathForCreate(FolderPath, *InNewName, ValidationErrorMsg))
+		const bool bAllowExistingFolder = !bCheckUniqueName;
+		if (!AssetViewUtils::IsValidFolderPathForCreate(FolderPath, *InNewName, ValidationErrorMsg, bAllowExistingFolder))
 		{
 			ContentBrowserAssetData::SetOptionalErrorMessage(OutErrorMsg, ValidationErrorMsg);
 			return false;
@@ -1002,6 +1015,12 @@ bool CanRenameAssetFolderItem(IAssetTools* InAssetTools, const FContentBrowserAs
 }
 
 bool CanRenameAssetFileItem(IAssetTools* InAssetTools, const FContentBrowserAssetFileItemDataPayload& InAssetPayload, const FString* InNewName, const bool InIsTempoarary, FText* OutErrorMsg)
+{
+	constexpr bool bCheckUniqueName = true;
+	return CanRenameAssetFileItem(InAssetTools, InAssetPayload, bCheckUniqueName, InNewName, InIsTempoarary, OutErrorMsg);
+}
+
+bool CanRenameAssetFileItem(IAssetTools* InAssetTools, const FContentBrowserAssetFileItemDataPayload& InAssetPayload, const bool bCheckUniqueName, const FString* InNewName, const bool InIsTempoarary, FText* OutErrorMsg)
 {
 	if (!CanModifyAssetFileItem(InAssetTools, InAssetPayload, OutErrorMsg))
 	{
@@ -1026,7 +1045,8 @@ bool CanRenameAssetFileItem(IAssetTools* InAssetTools, const FContentBrowserAsse
 		const FString ObjectPath = FString::Printf(TEXT("%s.%s"), *PackageName, **InNewName);
 
 		FText ValidationErrorMsg;
-		if (!AssetViewUtils::IsValidObjectPathForCreate(ObjectPath, ValidationErrorMsg))
+		const bool bAllowExistingAsset = !bCheckUniqueName;
+		if (!AssetViewUtils::IsValidObjectPathForCreate(ObjectPath, ValidationErrorMsg, bAllowExistingAsset))
 		{
 			SetOptionalErrorMessage(OutErrorMsg, ValidationErrorMsg);
 			return false;
@@ -1046,9 +1066,11 @@ bool CanRenameAssetFileItem(IAssetTools* InAssetTools, const FContentBrowserAsse
 
 bool RenameItem(IAssetTools* InAssetTools, IAssetRegistry* InAssetRegistry, const UContentBrowserDataSource* InOwnerDataSource, const FContentBrowserItemData& InItem, const FString& InNewName)
 {
+	constexpr bool bCheckUniqueName = false;
+
 	if (TSharedPtr<const FContentBrowserAssetFolderItemDataPayload> FolderPayload = GetAssetFolderItemPayload(InOwnerDataSource, InItem))
 	{
-		if (CanRenameAssetFolderItem(InAssetTools, *FolderPayload, &InNewName, nullptr))
+		if (CanRenameAssetFolderItem(InAssetTools, *FolderPayload, bCheckUniqueName, &InNewName, nullptr))
 		{
 			return RenameAssetFolderItem(InAssetRegistry, *FolderPayload, InNewName);
 		}
@@ -1056,7 +1078,7 @@ bool RenameItem(IAssetTools* InAssetTools, IAssetRegistry* InAssetRegistry, cons
 
 	if (TSharedPtr<const FContentBrowserAssetFileItemDataPayload> AssetPayload = GetAssetFileItemPayload(InOwnerDataSource, InItem))
 	{
-		if (CanRenameAssetFileItem(InAssetTools, *AssetPayload, &InNewName, InItem.IsTemporary(), nullptr))
+		if (CanRenameAssetFileItem(InAssetTools, *AssetPayload, bCheckUniqueName, &InNewName, InItem.IsTemporary(), nullptr))
 		{
 			return RenameAssetFileItem(InAssetTools, *AssetPayload, InNewName);
 		}
@@ -1074,7 +1096,8 @@ bool RenameAssetFolderItem(IAssetRegistry* InAssetRegistry, const FContentBrowse
 	FString NewPathOnDisk;
 	if (FPackageName::TryConvertLongPackageNameToFilename(NewPath, NewPathOnDisk) && IFileManager::Get().MakeDirectory(*NewPathOnDisk, true))
 	{
-		if (InAssetRegistry->AddPath(NewPath) && AssetViewUtils::RenameFolder(NewPath, OldPath))
+		InAssetRegistry->AddPath(NewPath);
+		if (AssetViewUtils::RenameFolder(NewPath, OldPath))
 		{
 			return true;
 		}
diff --git a/Engine/Plugins/Editor/ContentBrowser/ContentBrowserAssetDataSource/Source/ContentBrowserAssetDataSource/Private/ContentBrowserAssetDataSource.cpp b/Engine/Plugins/Editor/ContentBrowser/ContentBrowserAssetDataSource/Source/ContentBrowserAssetDataSource/Private/ContentBrowserAssetDataSource.cpp
--- a/Engine/Plugins/Editor/ContentBrowser/ContentBrowserAssetDataSource/Source/ContentBrowserAssetDataSource/Private/ContentBrowserAssetDataSource.cpp
+++ b/Engine/Plugins/Editor/ContentBrowser/ContentBrowserAssetDataSource/Source/ContentBrowserAssetDataSource/Private/ContentBrowserAssetDataSource.cpp
@@ -2350,7 +2350,7 @@ bool UContentBrowserAssetDataSource::IsFolderVisible(const FName InPath, const E
 		if (InContentsFilter.HideFolderIfEmptyFilter)
 		{
 			FNameBuilder InternalFolderPathBuilder(InternalFolderPath);
-			return !InContentsFilter.HideFolderIfEmptyFilter->HideFolderIfEmpty(InternalFolderPath, InternalFolderPathBuilder);
+			return !HideFolderIfEmpty(*InContentsFilter.HideFolderIfEmptyFilter, InternalFolderPath, InternalFolderPathBuilder);
 		}
 		return false;
 	};
@@ -2403,7 +2403,7 @@ bool UContentBrowserAssetDataSource::CanCreateFolder(const FName InPath, FText*
 	return ContentBrowserAssetData::CanModifyPath(AssetTools, InternalPath, OutErrorMsg);
 }
 
-bool UContentBrowserAssetDataSource::CreateFolder(const FName InPath, FContentBrowserItemDataTemporaryContext& OutPendingItem)
+bool UContentBrowserAssetDataSource::CreateFolder(const FName InPath, const TSharedPtr<IContentBrowserHideFolderIfEmptyFilter>& HideFolderIfEmptyFilter, FContentBrowserItemDataTemporaryContext& OutPendingItem)
 {
 	const FString ParentPath = FPackageName::GetLongPackagePath(InPath.ToString());
 	FName InternalParentPath;
@@ -2427,7 +2427,14 @@ bool UContentBrowserAssetDataSource::CreateFolder(const FName InPath, FContentBr
 
 	OutPendingItem = FContentBrowserItemDataTemporaryContext(
 		MoveTemp(NewItemData),
-		FContentBrowserItemDataTemporaryContext::FOnValidateItem::CreateUObject(this, &UContentBrowserAssetDataSource::OnValidateItemName),
+		FContentBrowserItemDataTemporaryContext::FOnValidateItem::CreateLambda([DataSource = TWeakObjectPtr<UContentBrowserAssetDataSource>(this), HideFolderIfEmptyFilter](const FContentBrowserItemData& InItem, const FString& InProposedName, FText* OutErrorMsg)
+			{
+				if (UContentBrowserAssetDataSource* DataSourcePtr = DataSource.Get())
+				{
+					return DataSourcePtr->CanRenameItem(InItem, &InProposedName, HideFolderIfEmptyFilter.Get(), OutErrorMsg);
+				}
+				return false;
+			}),
 		FContentBrowserItemDataTemporaryContext::FOnFinalizeItem::CreateUObject(this, &UContentBrowserAssetDataSource::OnFinalizeCreateFolder)
 		);
 
@@ -2766,9 +2773,28 @@ bool UContentBrowserAssetDataSource::BulkPrivatizeItems(TArrayView<const FConten
 	return ContentBrowserAssetData::PrivatizeItems(AssetTools, AssetRegistry, this, InItems);
 }
 
-bool UContentBrowserAssetDataSource::CanRenameItem(const FContentBrowserItemData& InItem, const FString* InNewName, FText* OutErrorMsg)
+bool UContentBrowserAssetDataSource::CanRenameItem(const FContentBrowserItemData& InItem, const FString* InNewName, const IContentBrowserHideFolderIfEmptyFilter* HideFolderIfEmptyFilter, FText* OutErrorMsg)
 {
-	return ContentBrowserAssetData::CanRenameItem(AssetTools, this, InItem, InNewName, OutErrorMsg);
+	bool bCheckUniqueName = true;
+
+	if (InNewName != nullptr && HideFolderIfEmptyFilter != nullptr)
+	{
+		if (TSharedPtr<const FContentBrowserAssetFolderItemDataPayload> FolderPayload = GetAssetFolderItemPayload(InItem))
+		{
+			FString NewInternalPathString = FPaths::GetPath(FolderPayload->GetInternalPath().ToString()) / *InNewName;
+			FName NewInternalPath(NewInternalPathString);
+
+			const EContentBrowserFolderAttributes FolderAttributes = GetAssetFolderAttributes(NewInternalPath);
+			if (!EnumHasAnyFlags(FolderAttributes, EContentBrowserFolderAttributes::AlwaysVisible | EContentBrowserFolderAttributes::HasAssets | EContentBrowserFolderAttributes::HasRedirectors) &&
+				HideFolderIfEmpty(*HideFolderIfEmptyFilter, NewInternalPath, NewInternalPathString))
+			{
+				// We are renaming to an existing hidden folder name, disable the unique name check since it will fail.
+				bCheckUniqueName = false;
+			}
+		}
+	}
+
+	return ContentBrowserAssetData::CanRenameItem(AssetTools, this, InItem, bCheckUniqueName, InNewName, OutErrorMsg);
 }
 
 bool UContentBrowserAssetDataSource::RenameItem(const FContentBrowserItemData& InItem, const FString& InNewName, FContentBrowserItemData& OutNewItem)
@@ -2780,8 +2806,9 @@ bool UContentBrowserAssetDataSource::RenameItem(const FContentBrowserItemData& I
 		case EContentBrowserItemFlags::Type_Folder:
 			if (TSharedPtr<const FContentBrowserAssetFolderItemDataPayload> FolderPayload = GetAssetFolderItemPayload(InItem))
 			{
-				const FName NewFolderPath = *(FPaths::GetPath(FolderPayload->GetInternalPath().ToString()) / InNewName);
-				OutNewItem = CreateAssetFolderItem(NewFolderPath);
+				const FString NewFolderPath = FPaths::GetPath(FolderPayload->GetInternalPath().ToString()) / InNewName;
+				AssetViewUtils::OnAlwaysShowPath().Broadcast(NewFolderPath);
+				OutNewItem = CreateAssetFolderItem(FName(NewFolderPath));
 			}
 			break;
 
@@ -3592,6 +3619,28 @@ bool UContentBrowserAssetDataSource::ClearAssetFolderAttributes(const FName InPa
 	return false;
 }
 
+bool UContentBrowserAssetDataSource::HideFolderIfEmpty(const IContentBrowserHideFolderIfEmptyFilter& HideFolderIfEmptyFilter, FName Path, FStringView PathString) const
+{
+	if (!HideFolderIfEmptyFilter.HideFolderIfEmpty(Path, PathString))
+	{
+		return false;
+	}
+
+	// If any subpaths shouldn't be hidden, then Path should be visible.
+	bool bAnySubPathVisible = false;
+	AssetRegistry->EnumerateSubPaths(Path, [&HideFolderIfEmptyFilter, &bAnySubPathVisible](FName ChildPath)
+		{
+			FNameBuilder ChildPathBuilder(ChildPath);
+			if (!HideFolderIfEmptyFilter.HideFolderIfEmpty(ChildPath, ChildPathBuilder))
+			{
+				bAnySubPathVisible = true;
+				return false;
+			}
+			return true;
+		}, true);
+	return !bAnySubPathVisible;
+}
+
 void UContentBrowserAssetDataSource::PopulateAddNewContextMenu(UToolMenu* InMenu)
 {
 	const UContentBrowserDataMenuContext_AddNewMenu* ContextObject = InMenu->FindContext<UContentBrowserDataMenuContext_AddNewMenu>();
@@ -3856,7 +3905,7 @@ void UContentBrowserAssetDataSource::OnBeginCreateAsset(const FName InDefaultAss
 
 bool UContentBrowserAssetDataSource::OnValidateItemName(const FContentBrowserItemData& InItem, const FString& InProposedName, FText* OutErrorMsg)
 {
-	return CanRenameItem(InItem, &InProposedName, OutErrorMsg);
+	return CanRenameItem(InItem, &InProposedName, nullptr, OutErrorMsg);
 }
 
 FReply UContentBrowserAssetDataSource::OnImportClicked(const UContentBrowserToolbarMenuContext* ContextObject)
diff --git a/Engine/Plugins/Editor/ContentBrowser/ContentBrowserAssetDataSource/Source/ContentBrowserAssetDataSource/Public/ContentBrowserAssetDataCore.h b/Engine/Plugins/Editor/ContentBrowser/ContentBrowserAssetDataSource/Source/ContentBrowserAssetDataSource/Public/ContentBrowserAssetDataCore.h
--- a/Engine/Plugins/Editor/ContentBrowser/ContentBrowserAssetDataSource/Source/ContentBrowserAssetDataSource/Public/ContentBrowserAssetDataCore.h
+++ b/Engine/Plugins/Editor/ContentBrowser/ContentBrowserAssetDataSource/Source/ContentBrowserAssetDataSource/Public/ContentBrowserAssetDataCore.h
@@ -137,10 +137,13 @@ CONTENTBROWSERASSETDATASOURCE_API FContentBrowserItemData CreateUnsupportedAsset
 	CONTENTBROWSERASSETDATASOURCE_API bool PrivatizeItems(IAssetTools* InAssetTools, IAssetRegistry* InAssetRegistry, const UContentBrowserDataSource* InOwnerDataSource, TArrayView<const FContentBrowserItemData> InItems);
 
 	CONTENTBROWSERASSETDATASOURCE_API bool CanRenameItem(IAssetTools* InAssetTools, const UContentBrowserDataSource* InOwnerDataSource, const FContentBrowserItemData& InItem, const FString* InNewName, FText* OutErrorMsg);
+	CONTENTBROWSERASSETDATASOURCE_API bool CanRenameItem(IAssetTools* InAssetTools, const UContentBrowserDataSource* InOwnerDataSource, const FContentBrowserItemData& InItem, const bool bCheckUniqueName, const FString* InNewName, FText* OutErrorMsg);
 
 	CONTENTBROWSERASSETDATASOURCE_API bool CanRenameAssetFolderItem(IAssetTools* InAssetTools, const FContentBrowserAssetFolderItemDataPayload& InFolderPayload, const FString* InNewName, FText* OutErrorMsg);
+	CONTENTBROWSERASSETDATASOURCE_API bool CanRenameAssetFolderItem(IAssetTools* InAssetTools, const FContentBrowserAssetFolderItemDataPayload& InFolderPayload, const bool bCheckUniqueName, const FString* InNewName, FText* OutErrorMsg);
 
 	CONTENTBROWSERASSETDATASOURCE_API bool CanRenameAssetFileItem(IAssetTools* InAssetTools, const FContentBrowserAssetFileItemDataPayload& InAssetPayload, const FString* InNewName, const bool InIsTempoarary, FText* OutErrorMsg);
+	CONTENTBROWSERASSETDATASOURCE_API bool CanRenameAssetFileItem(IAssetTools* InAssetTools, const FContentBrowserAssetFileItemDataPayload& InAssetPayload, const bool bCheckUniqueName, const FString* InNewName, const bool InIsTempoarary, FText* OutErrorMsg);
 
 	CONTENTBROWSERASSETDATASOURCE_API bool RenameItem(IAssetTools* InAssetTools, IAssetRegistry* InAssetRegistry, const UContentBrowserDataSource* InOwnerDataSource, const FContentBrowserItemData& InItem, const FString& InNewName);
 
diff --git a/Engine/Plugins/Editor/ContentBrowser/ContentBrowserAssetDataSource/Source/ContentBrowserAssetDataSource/Public/ContentBrowserAssetDataSource.h b/Engine/Plugins/Editor/ContentBrowser/ContentBrowserAssetDataSource/Source/ContentBrowserAssetDataSource/Public/ContentBrowserAssetDataSource.h
--- a/Engine/Plugins/Editor/ContentBrowser/ContentBrowserAssetDataSource/Source/ContentBrowserAssetDataSource/Public/ContentBrowserAssetDataSource.h
+++ b/Engine/Plugins/Editor/ContentBrowser/ContentBrowserAssetDataSource/Source/ContentBrowserAssetDataSource/Public/ContentBrowserAssetDataSource.h
@@ -255,7 +255,7 @@ public:
 
 	virtual bool CanCreateFolder(const FName InPath, FText* OutErrorMsg) override;
 
-	virtual bool CreateFolder(const FName InPath, FContentBrowserItemDataTemporaryContext& OutPendingItem) override;
+	virtual bool CreateFolder(const FName InPath, const TSharedPtr<IContentBrowserHideFolderIfEmptyFilter>& HideFolderIfEmptyFilter, FContentBrowserItemDataTemporaryContext& OutPendingItem) override;
 
 	virtual bool DoesItemPassFilter(const FContentBrowserItemData& InItem, const FContentBrowserDataCompiledFilter& InFilter) override;
 
@@ -311,7 +311,7 @@ public:
 
 	virtual bool BulkPrivatizeItems(TArrayView<const FContentBrowserItemData> InItems) override;
 
-	virtual bool CanRenameItem(const FContentBrowserItemData& InItem, const FString* InNewName, FText* OutErrorMsg) override;
+	virtual bool CanRenameItem(const FContentBrowserItemData& InItem, const FString* InNewName, const IContentBrowserHideFolderIfEmptyFilter* HideFolderIfEmptyFilter, FText* OutErrorMsg) override;
 
 	virtual bool RenameItem(const FContentBrowserItemData& InItem, const FString& InNewName, FContentBrowserItemData& OutNewItem) override;
 
@@ -417,6 +417,8 @@ private:
 
 	bool ClearAssetFolderAttributes(const FName InPath, const EContentBrowserFolderAttributes InAttributesToClear);
 
+	bool HideFolderIfEmpty(const IContentBrowserHideFolderIfEmptyFilter& HideFolderIfEmptyFilter, FName Path, FStringView PathString) const;
+
 	void PopulateAddNewContextMenu(UToolMenu* InMenu);
 
 	void PopulateContentBrowserToolBar(UToolMenu* InMenu);
diff --git a/Engine/Plugins/Editor/ContentBrowser/ContentBrowserFileDataSource/Source/ContentBrowserFileDataSource/Private/ContentBrowserFileDataCore.cpp b/Engine/Plugins/Editor/ContentBrowser/ContentBrowserFileDataSource/Source/ContentBrowserFileDataSource/Private/ContentBrowserFileDataCore.cpp
--- a/Engine/Plugins/Editor/ContentBrowser/ContentBrowserFileDataSource/Source/ContentBrowserFileDataSource/Private/ContentBrowserFileDataCore.cpp
+++ b/Engine/Plugins/Editor/ContentBrowser/ContentBrowserFileDataSource/Source/ContentBrowserFileDataSource/Private/ContentBrowserFileDataCore.cpp
@@ -658,15 +658,21 @@ bool DeleteFileItems(TArrayView<const TSharedRef<const FContentBrowserFileItemDa
 }
 
 bool CanRenameItem(const UContentBrowserDataSource* InOwnerDataSource, const FContentBrowserItemData& InItem, const FString* InNewName, FText* OutErrorMsg)
+{
+	constexpr bool bCheckUniqueName = true;
+	return CanRenameItem(InOwnerDataSource, InItem, bCheckUniqueName, InNewName, OutErrorMsg);
+}
+
+bool CanRenameItem(const UContentBrowserDataSource* InOwnerDataSource, const FContentBrowserItemData& InItem, const bool bCheckUniqueName, const FString* InNewName, FText* OutErrorMsg)
 {
 	if (TSharedPtr<const FContentBrowserFolderItemDataPayload> FolderPayload = GetFolderItemPayload(InOwnerDataSource, InItem))
 	{
-		return CanRenameFolderItem(*FolderPayload, /*bCheckUniqueName*/true, InNewName, OutErrorMsg);
+		return CanRenameFolderItem(*FolderPayload, bCheckUniqueName, InNewName, OutErrorMsg);
 	}
 
 	if (TSharedPtr<const FContentBrowserFileItemDataPayload> FilePayload = GetFileItemPayload(InOwnerDataSource, InItem))
 	{
-		return CanRenameFileItem(*FilePayload, /*bCheckUniqueName*/true, InNewName, OutErrorMsg);
+		return CanRenameFileItem(*FilePayload, bCheckUniqueName, InNewName, OutErrorMsg);
 	}
 
 	return false;
diff --git a/Engine/Plugins/Editor/ContentBrowser/ContentBrowserFileDataSource/Source/ContentBrowserFileDataSource/Private/ContentBrowserFileDataSource.cpp b/Engine/Plugins/Editor/ContentBrowser/ContentBrowserFileDataSource/Source/ContentBrowserFileDataSource/Private/ContentBrowserFileDataSource.cpp
--- a/Engine/Plugins/Editor/ContentBrowser/ContentBrowserFileDataSource/Source/ContentBrowserFileDataSource/Private/ContentBrowserFileDataSource.cpp
+++ b/Engine/Plugins/Editor/ContentBrowser/ContentBrowserFileDataSource/Source/ContentBrowserFileDataSource/Private/ContentBrowserFileDataSource.cpp
@@ -1003,7 +1003,7 @@ bool UContentBrowserFileDataSource::CanCreateFolder(const FName InPath, FText* O
 	return ContentBrowserFileData::CanModifyDirectory(InternalDiskPath, OutErrorMsg);
 }
 
-bool UContentBrowserFileDataSource::CreateFolder(const FName InPath, FContentBrowserItemDataTemporaryContext& OutPendingItem)
+bool UContentBrowserFileDataSource::CreateFolder(const FName InPath, const TSharedPtr<IContentBrowserHideFolderIfEmptyFilter>& HideFolderIfEmptyFilter, FContentBrowserItemDataTemporaryContext& OutPendingItem)
 {
 	FName InternalPath;
 	if (!TryConvertVirtualPathToInternal(InPath, InternalPath))
@@ -1031,8 +1031,14 @@ bool UContentBrowserFileDataSource::CreateFolder(const FName InPath, FContentBro
 
 	OutPendingItem = FContentBrowserItemDataTemporaryContext(
 		MoveTemp(NewItemData),
-		FContentBrowserItemDataTemporaryContext::FOnValidateItem::CreateUObject(
-			this, &UContentBrowserFileDataSource::OnValidateItemName),
+		FContentBrowserItemDataTemporaryContext::FOnValidateItem::CreateLambda([DataSource = TWeakObjectPtr<UContentBrowserFileDataSource>(this), HideFolderIfEmptyFilter](const FContentBrowserItemData& InItem, const FString& InProposedName, FText* OutErrorMsg)
+			{
+				if (UContentBrowserFileDataSource* DataSourcePtr = DataSource.Get())
+				{
+					return DataSourcePtr->CanRenameItem(InItem, &InProposedName, HideFolderIfEmptyFilter.Get(), OutErrorMsg);
+				}
+				return false;
+			}),
 		FContentBrowserItemDataTemporaryContext::FOnFinalizeItem::CreateUObject(
 			this, &UContentBrowserFileDataSource::OnFinalizeCreateFolder));
 
@@ -1104,7 +1110,7 @@ bool UContentBrowserFileDataSource::IsFolderVisible(const FName InPath, const EC
 	{
 		if (ContentsFilter.HideFolderIfEmptyFilter)
 		{
-			return !ContentsFilter.HideFolderIfEmptyFilter->HideFolderIfEmpty(ConvertedPath, InternalPathStrView);
+			return !HideFolderIfEmpty(*ContentsFilter.HideFolderIfEmptyFilter, ConvertedPath, InternalPathStrView);
 		}
 		return false;
 	}
@@ -1361,7 +1367,7 @@ bool UContentBrowserFileDataSource::BulkDeleteItems(TArrayView<const FContentBro
 	return ContentBrowserFileData::DeleteItems(this, InItems);
 }
 
-bool UContentBrowserFileDataSource::CanRenameItem(const FContentBrowserItemData& InItem, const FString* InNewName, FText* OutErrorMsg)
+bool UContentBrowserFileDataSource::CanRenameItem(const FContentBrowserItemData& InItem, const FString* InNewName, const IContentBrowserHideFolderIfEmptyFilter* HideFolderIfEmptyFilter, FText* OutErrorMsg)
 {
 	// Cannot copy an item outside the paths known to this data source
 	FName InternalPath;
@@ -1384,7 +1390,26 @@ bool UContentBrowserFileDataSource::CanRenameItem(const FContentBrowserItemData&
 		return false;
 	}
 
-	return ContentBrowserFileData::CanRenameItem(this, InItem, InNewName, OutErrorMsg);
+	bool bCheckUniqueName = true;
+
+	if (InNewName != nullptr && HideFolderIfEmptyFilter != nullptr)
+	{
+		if (TSharedPtr<const FContentBrowserFolderItemDataPayload> FolderPayload = GetFolderItemPayload(InItem))
+		{
+			FString NewInternalPath = FPaths::GetPath(FolderPayload->GetInternalPath().ToString()) / *InNewName;
+			const uint32 NewInternalPathHash = GetTypeHash(NewInternalPath);
+
+			if (!AlwaysVisibleFolders.ContainsByHash(NewInternalPathHash, NewInternalPath) &&
+				EmptyFolders.ContainsByHash(NewInternalPathHash, NewInternalPath) &&
+				HideFolderIfEmpty(*HideFolderIfEmptyFilter, FName(NewInternalPath), NewInternalPath))
+			{
+				// We are renaming to an existing hidden folder name, disable the unique name check since it will fail.
+				bCheckUniqueName = false;
+			}
+		}
+	}
+
+	return ContentBrowserFileData::CanRenameItem(this, InItem, bCheckUniqueName, InNewName, OutErrorMsg);
 }
 
 bool UContentBrowserFileDataSource::RenameItem(const FContentBrowserItemData& InItem, const FString& InNewName, FContentBrowserItemData& OutNewItem)
@@ -1618,6 +1643,33 @@ TSharedPtr<const FContentBrowserFileItemDataPayload> UContentBrowserFileDataSour
 	return ContentBrowserFileData::GetFileItemPayload(this, InItem);
 }
 
+bool UContentBrowserFileDataSource::HideFolderIfEmpty(const IContentBrowserHideFolderIfEmptyFilter& HideFolderIfEmptyFilter, FName Path, FStringView PathString) const
+{
+	check(EmptyFolders.ContainsByHash(GetTypeHash(PathString), PathString));
+
+	if (!HideFolderIfEmptyFilter.HideFolderIfEmpty(Path, PathString))
+	{
+		return false;
+	}
+
+	if (const FDiscoveredItem* DiscoveredItem = DiscoveredItems.Find(Path))
+	{
+		checkf(DiscoveredItem->Type == FDiscoveredItem::EType::Directory, TEXT("Empty folders should not contain files."));
+
+		for (const FName ChildItem : DiscoveredItem->ChildItems)
+		{
+			FNameBuilder ChildPathBuilder(ChildItem);
+			if (!HideFolderIfEmpty(HideFolderIfEmptyFilter, ChildItem, ChildPathBuilder))
+			{
+				// If any subpaths shouldn't be hidden, then Path should be visible.
+				return false;
+			}
+		}
+	}
+
+	return true;
+}
+
 bool UContentBrowserFileDataSource::IsKnownFileMount(const FName InMountPath, FString* OutDiskPath) const
 {
 	FNameBuilder MountPathStr(InMountPath);
@@ -2066,7 +2118,7 @@ void UContentBrowserFileDataSource::OnNewFileRequested(const FName InDestFolderP
 
 bool UContentBrowserFileDataSource::OnValidateItemName(const FContentBrowserItemData& InItem, const FString& InProposedName, FText* OutErrorMsg)
 {
-	return CanRenameItem(InItem, &InProposedName, OutErrorMsg);
+	return CanRenameItem(InItem, &InProposedName, nullptr, OutErrorMsg);
 }
 
 FContentBrowserItemData UContentBrowserFileDataSource::OnFinalizeCreateFolder(const FContentBrowserItemData& InItemData, const FString& InProposedName, FText* OutErrorMsg)
diff --git a/Engine/Plugins/Editor/ContentBrowser/ContentBrowserFileDataSource/Source/ContentBrowserFileDataSource/Public/ContentBrowserFileDataCore.h b/Engine/Plugins/Editor/ContentBrowser/ContentBrowserFileDataSource/Source/ContentBrowserFileDataSource/Public/ContentBrowserFileDataCore.h
--- a/Engine/Plugins/Editor/ContentBrowser/ContentBrowserFileDataSource/Source/ContentBrowserFileDataSource/Public/ContentBrowserFileDataCore.h
+++ b/Engine/Plugins/Editor/ContentBrowser/ContentBrowserFileDataSource/Source/ContentBrowserFileDataSource/Public/ContentBrowserFileDataCore.h
@@ -87,6 +87,7 @@ namespace ContentBrowserFileData
 	CONTENTBROWSERFILEDATASOURCE_API bool DeleteFileItems(TArrayView<const TSharedRef<const FContentBrowserFileItemDataPayload>> InFilePayloads);
 
 	CONTENTBROWSERFILEDATASOURCE_API bool CanRenameItem(const UContentBrowserDataSource* InOwnerDataSource, const FContentBrowserItemData& InItem, const FString* InNewName, FText* OutErrorMsg);
+	CONTENTBROWSERFILEDATASOURCE_API bool CanRenameItem(const UContentBrowserDataSource* InOwnerDataSource, const FContentBrowserItemData& InItem, const bool bCheckUniqueName, const FString* InNewName, FText* OutErrorMsg);
 
 	CONTENTBROWSERFILEDATASOURCE_API bool CanRenameFolderItem(const FContentBrowserFolderItemDataPayload& InFolderPayload, const bool bCheckUniqueName, const FString* InNewName, FText* OutErrorMsg);
 
diff --git a/Engine/Plugins/Editor/ContentBrowser/ContentBrowserFileDataSource/Source/ContentBrowserFileDataSource/Public/ContentBrowserFileDataSource.h b/Engine/Plugins/Editor/ContentBrowser/ContentBrowserFileDataSource/Source/ContentBrowserFileDataSource/Public/ContentBrowserFileDataSource.h
--- a/Engine/Plugins/Editor/ContentBrowser/ContentBrowserFileDataSource/Source/ContentBrowserFileDataSource/Public/ContentBrowserFileDataSource.h
+++ b/Engine/Plugins/Editor/ContentBrowser/ContentBrowserFileDataSource/Source/ContentBrowserFileDataSource/Public/ContentBrowserFileDataSource.h
@@ -53,7 +53,7 @@ public:
 
 	virtual bool CanCreateFolder(const FName InPath, FText* OutErrorMsg) override;
 
-	virtual bool CreateFolder(const FName InPath, FContentBrowserItemDataTemporaryContext& OutPendingItem) override;
+	virtual bool CreateFolder(const FName InPath, const TSharedPtr<IContentBrowserHideFolderIfEmptyFilter>& HideFolderIfEmptyFilter, FContentBrowserItemDataTemporaryContext& OutPendingItem) override;
 
 	virtual bool IsDiscoveringItems(FText* OutStatus = nullptr) override;
 
@@ -93,7 +93,7 @@ public:
 
 	virtual bool BulkDeleteItems(TArrayView<const FContentBrowserItemData> InItems) override;
 
-	virtual bool CanRenameItem(const FContentBrowserItemData& InItem, const FString* InNewName, FText* OutErrorMsg) override;
+	virtual bool CanRenameItem(const FContentBrowserItemData& InItem, const FString* InNewName, const IContentBrowserHideFolderIfEmptyFilter* HideFolderIfEmptyFilter, FText* OutErrorMsg) override;
 
 	virtual bool RenameItem(const FContentBrowserItemData& InItem, const FString& InNewName, FContentBrowserItemData& OutNewItem) override;
 
@@ -144,6 +144,8 @@ protected:
 
 	TSharedPtr<const FContentBrowserFileItemDataPayload> GetFileItemPayload(const FContentBrowserItemData& InItem) const;
 
+	bool HideFolderIfEmpty(const IContentBrowserHideFolderIfEmptyFilter& HideFolderIfEmptyFilter, FName Path, FStringView PathString) const;
+
 	bool IsKnownFileMount(const FName InMountPath, FString* OutDiskPath = nullptr) const;
 
 	bool IsRootFileMount(const FName InMountPath, FString* OutDiskPath = nullptr) const;
diff --git a/Engine/Source/Developer/AssetTools/Private/AssetViewUtils.cpp b/Engine/Source/Developer/AssetTools/Private/AssetViewUtils.cpp
--- a/Engine/Source/Developer/AssetTools/Private/AssetViewUtils.cpp
+++ b/Engine/Source/Developer/AssetTools/Private/AssetViewUtils.cpp
@@ -1508,7 +1508,7 @@ bool AssetViewUtils::IsValidObjectPathForCreate(const FString& ObjectPath, const
 	return true;
 }
 
-bool AssetViewUtils::IsValidFolderPathForCreate(const FString& InFolderPath, const FString& NewFolderName, FText& OutErrorMessage)
+bool AssetViewUtils::IsValidFolderPathForCreate(const FString& InFolderPath, const FString& NewFolderName, FText& OutErrorMessage, bool bAllowExistingFolder)
 {
 	if (!IsValidFolderName(NewFolderName, OutErrorMessage))
 	{
@@ -1517,7 +1517,7 @@ bool AssetViewUtils::IsValidFolderPathForCreate(const FString& InFolderPath, con
 
 	const FString NewFolderPath = InFolderPath / NewFolderName;
 
-	if (DoesFolderExist(NewFolderPath))
+	if (!bAllowExistingFolder && DoesFolderExist(NewFolderPath))
 	{
 		OutErrorMessage = LOCTEXT("RenameFolderAlreadyExists", "A folder already exists at this location with this name.");
 		return false;
diff --git a/Engine/Source/Developer/AssetTools/Public/AssetViewUtils.h b/Engine/Source/Developer/AssetTools/Public/AssetViewUtils.h
--- a/Engine/Source/Developer/AssetTools/Public/AssetViewUtils.h
+++ b/Engine/Source/Developer/AssetTools/Public/AssetViewUtils.h
@@ -229,7 +229,7 @@ namespace AssetViewUtils
 	ASSETTOOLS_API bool IsValidObjectPathForCreate(const FString& ObjectPath, const UClass* ObjectClass, FText& OutErrorMessage, bool bAllowExistingAsset = false);
 
 	/** Returns true if the specified folder name in the specified path is available for folder creation */
-	ASSETTOOLS_API bool IsValidFolderPathForCreate(const FString& FolderPath, const FString& NewFolderName, FText& OutErrorMessage);
+	ASSETTOOLS_API bool IsValidFolderPathForCreate(const FString& FolderPath, const FString& NewFolderName, FText& OutErrorMessage, bool bAllowExistingFolder = false);
 
 	/** Returns the relative path, from the workspace root, of the package */
 	ASSETTOOLS_API FString GetPackagePathWithinRoot(const FString& PackageName);
diff --git a/Engine/Source/Editor/ContentBrowser/Private/SAssetView.cpp b/Engine/Source/Editor/ContentBrowser/Private/SAssetView.cpp
--- a/Engine/Source/Editor/ContentBrowser/Private/SAssetView.cpp
+++ b/Engine/Source/Editor/ContentBrowser/Private/SAssetView.cpp
@@ -5308,7 +5308,9 @@ bool SAssetView::AssetVerifyRenameCommit(const TSharedPtr<FAssetViewItem>& Item,
 	}
 	else if (!Item->GetItem().GetItemName().ToString().Equals(NewItemName))
 	{
-		return Item->GetItem().CanRename(&NewItemName, &OutErrorMessage);
+		UContentBrowserDataSubsystem* ContentBrowserData = IContentBrowserDataModule::Get().GetSubsystem();
+
+		return Item->GetItem().CanRename(&NewItemName, ContentBrowserData ->CreateHideFolderIfEmptyFilter().Get(), &OutErrorMessage);
 	}
 
 	return true;
@@ -5351,7 +5353,7 @@ void SAssetView::AssetRenameCommit(const TSharedPtr<FAssetViewItem>& Item, const
 		FScopedSuppressContentBrowserDataTick TickSuppression(ContentBrowserData);
 
 		FContentBrowserItem NewItem;
-		if (Item->GetItem().CanRename(&NewName, &ErrorMessage) && Item->GetItem().Rename(NewName, &NewItem))
+		if (Item->GetItem().CanRename(&NewName, ContentBrowserData->CreateHideFolderIfEmptyFilter().Get(), &ErrorMessage) && Item->GetItem().Rename(NewName, &NewItem))
 		{
 			// Add result to view (the old item will be removed via the notifications, as not all data sources may have been able to perform the rename)
 			UpdatedItem = Items->CreateItemFromUser(MoveTemp(NewItem), FilteredAssetItems);
diff --git a/Engine/Source/Editor/ContentBrowser/Private/SPathView.cpp b/Engine/Source/Editor/ContentBrowser/Private/SPathView.cpp
--- a/Engine/Source/Editor/ContentBrowser/Private/SPathView.cpp
+++ b/Engine/Source/Editor/ContentBrowser/Private/SPathView.cpp
@@ -2296,7 +2296,9 @@ bool SPathView::VerifyFolderNameChanged(const TSharedPtr< FTreeItem >& TreeItem,
 	}
 	else if (!TreeItem->GetItem().GetItemName().ToString().Equals(ProposedName))
 	{
-		return TreeItem->GetItem().CanRename(&ProposedName, &OutErrorMessage);
+		UContentBrowserDataSubsystem* ContentBrowserData = IContentBrowserDataModule::Get().GetSubsystem();
+
+		return TreeItem->GetItem().CanRename(&ProposedName, ContentBrowserData->CreateHideFolderIfEmptyFilter().Get(), &OutErrorMessage);
 	}
 
 	return true;
@@ -2360,8 +2362,8 @@ void SPathView::FolderNameChanged(const TSharedPtr<FTreeItem>& TreeItem,
 	{
 		UContentBrowserDataSubsystem* ContentBrowserData = IContentBrowserDataModule::Get().GetSubsystem();
 		FScopedSuppressContentBrowserDataTick TickSuppression(ContentBrowserData);
-
-		if (TreeItem->GetItem().CanRename(&ProposedName, &ErrorMessage) && TreeItem->GetItem().Rename(ProposedName, &NewItem))
+		
+		if (TreeItem->GetItem().CanRename(&ProposedName, ContentBrowserData->CreateHideFolderIfEmptyFilter().Get(), &ErrorMessage) && TreeItem->GetItem().Rename(ProposedName, &NewItem))
 		{
 			bSuccess = true;
 		}
diff --git a/Engine/Source/Editor/ContentBrowserData/Private/ContentBrowserDataSource.cpp b/Engine/Source/Editor/ContentBrowserData/Private/ContentBrowserDataSource.cpp
--- a/Engine/Source/Editor/ContentBrowserData/Private/ContentBrowserDataSource.cpp
+++ b/Engine/Source/Editor/ContentBrowserData/Private/ContentBrowserDataSource.cpp
@@ -273,7 +273,7 @@ bool UContentBrowserDataSource::CanCreateFolder(const FName InPath, FText* OutEr
 	return false;
 }
 
-bool UContentBrowserDataSource::CreateFolder(const FName InPath, FContentBrowserItemDataTemporaryContext& OutPendingItem)
+bool UContentBrowserDataSource::CreateFolder(const FName InPath, const TSharedPtr<IContentBrowserHideFolderIfEmptyFilter>& HideFolderIfEmptyFilter, FContentBrowserItemDataTemporaryContext& OutPendingItem)
 {
 	return false;
 }
@@ -443,7 +443,7 @@ bool UContentBrowserDataSource::BulkPrivatizeItems(TArrayView<const FContentBrow
 	return bSuccess;
 }
 
-bool UContentBrowserDataSource::CanRenameItem(const FContentBrowserItemData& InItem, const FString* InNewName, FText* OutErrorMsg)
+bool UContentBrowserDataSource::CanRenameItem(const FContentBrowserItemData& InItem, const FString* InNewName, const IContentBrowserHideFolderIfEmptyFilter* HideFolderIfEmptyFilter, FText* OutErrorMsg)
 {
 	return false;
 }
diff --git a/Engine/Source/Editor/ContentBrowserData/Private/ContentBrowserDataSubsystem.cpp b/Engine/Source/Editor/ContentBrowserData/Private/ContentBrowserDataSubsystem.cpp
--- a/Engine/Source/Editor/ContentBrowserData/Private/ContentBrowserDataSubsystem.cpp
+++ b/Engine/Source/Editor/ContentBrowserData/Private/ContentBrowserDataSubsystem.cpp
@@ -630,6 +630,14 @@ bool UContentBrowserDataSubsystem::PrioritizeSearchPath(const FName InPath)
 	return bDidPrioritize;
 }
 
+bool UContentBrowserDataSubsystem::IsFolderVisible(const FName InPath, const EContentBrowserIsFolderVisibleFlags InFlags) const
+{
+	FContentBrowserFolderContentsFilter ContentsFilter;
+	ContentsFilter.HideFolderIfEmptyFilter = CreateHideFolderIfEmptyFilter();
+
+	return IsFolderVisible(InPath, InFlags, ContentsFilter);
+}
+
 bool UContentBrowserDataSubsystem::IsFolderVisible(const FName InPath, const EContentBrowserIsFolderVisibleFlags InFlags, const FContentBrowserFolderContentsFilter& InContentsFilter) const
 {
 	bool bIsKnownPath = false;
@@ -674,6 +682,11 @@ bool UContentBrowserDataSubsystem::CanCreateFolder(const FName InPath, FText* Ou
 }
 
 FContentBrowserItemTemporaryContext UContentBrowserDataSubsystem::CreateFolder(const FName InPath) const
+{
+	return CreateFolder(InPath, CreateHideFolderIfEmptyFilter());
+}
+
+FContentBrowserItemTemporaryContext UContentBrowserDataSubsystem::CreateFolder(const FName InPath, const TSharedPtr<IContentBrowserHideFolderIfEmptyFilter>& HideFolderIfEmptyFilter) const
 {
 	FContentBrowserItemTemporaryContext NewItem;
 	for (const auto& ActiveDataSourcePair : ActiveDataSources)
@@ -682,7 +695,7 @@ FContentBrowserItemTemporaryContext UContentBrowserDataSubsystem::CreateFolder(c
 		if (DataSource->IsVirtualPathUnderMountRoot(InPath))
 		{
 			FContentBrowserItemDataTemporaryContext NewItemData;
-			if (DataSource->CreateFolder(InPath, NewItemData))
+			if (DataSource->CreateFolder(InPath, HideFolderIfEmptyFilter, NewItemData))
 			{
 				NewItem.AppendContext(MoveTemp(NewItemData));
 			}
diff --git a/Engine/Source/Editor/ContentBrowserData/Private/ContentBrowserItem.cpp b/Engine/Source/Editor/ContentBrowserData/Private/ContentBrowserItem.cpp
--- a/Engine/Source/Editor/ContentBrowserData/Private/ContentBrowserItem.cpp
+++ b/Engine/Source/Editor/ContentBrowserData/Private/ContentBrowserItem.cpp
@@ -524,7 +524,20 @@ bool FContentBrowserItem::Delete() const
 
 bool FContentBrowserItem::CanRename(const FString* InNewName, FText* OutErrorMsg) const
 {
-	return FContentBrowserItemHelper::CallDataSourceImpl<UContentBrowserDataSource>(*this, UE_PROJECTION_MEMBER(UContentBrowserDataSource, CanRenameItem), InNewName, OutErrorMsg);
+	TSharedPtr<IContentBrowserHideFolderIfEmptyFilter> HideFolderIfEmptyFilter;
+	if (InNewName != nullptr)
+	{
+		// Only create the filter if checking against a specific name, otherwise we shouldn't need it.
+		const UContentBrowserDataSubsystem* ContentBrowserData = IContentBrowserDataModule::Get().GetSubsystem();
+		HideFolderIfEmptyFilter = ContentBrowserData->CreateHideFolderIfEmptyFilter();
+	}
+
+	return CanRename(InNewName, HideFolderIfEmptyFilter.Get(), OutErrorMsg);
+}
+
+bool FContentBrowserItem::CanRename(const FString* InNewName, const IContentBrowserHideFolderIfEmptyFilter* HideFolderIfEmptyFilter, FText* OutErrorMsg) const
+{
+	return FContentBrowserItemHelper::CallDataSourceImpl<UContentBrowserDataSource>(*this, UE_PROJECTION_MEMBER(UContentBrowserDataSource, CanRenameItem), InNewName, HideFolderIfEmptyFilter, OutErrorMsg);
 }
 
 bool FContentBrowserItem::Rename(const FString& InNewName, FContentBrowserItem* OutNewItem) const
diff --git a/Engine/Source/Editor/ContentBrowserData/Public/ContentBrowserDataSource.h b/Engine/Source/Editor/ContentBrowserData/Public/ContentBrowserDataSource.h
--- a/Engine/Source/Editor/ContentBrowserData/Public/ContentBrowserDataSource.h
+++ b/Engine/Source/Editor/ContentBrowserData/Public/ContentBrowserDataSource.h
@@ -283,7 +283,7 @@ public:
 	 *
 	 * @return True if the pending folder was created, false otherwise.
 	 */
-	virtual bool CreateFolder(const FName InPath, FContentBrowserItemDataTemporaryContext& OutPendingItem);
+	virtual bool CreateFolder(const FName InPath, const TSharedPtr<IContentBrowserHideFolderIfEmptyFilter>& HideFolderIfEmptyFilter, FContentBrowserItemDataTemporaryContext& OutPendingItem);
 
 	/*
 	 * Query whether the given item passes the given compiled filter. Should be called after ConvertItemForFilter
@@ -570,7 +570,7 @@ public:
 	 *
 	 * @return True if the item can be renamed, false otherwise.
 	 */
-	virtual bool CanRenameItem(const FContentBrowserItemData& InItem, const FString* InNewName, FText* OutErrorMsg);
+	virtual bool CanRenameItem(const FContentBrowserItemData& InItem, const FString* InNewName, const IContentBrowserHideFolderIfEmptyFilter* HideFolderIfEmptyFilter, FText* OutErrorMsg);
 
 	/**
 	 * Attempt to rename the given item.
diff --git a/Engine/Source/Editor/ContentBrowserData/Public/ContentBrowserDataSubsystem.h b/Engine/Source/Editor/ContentBrowserData/Public/ContentBrowserDataSubsystem.h
--- a/Engine/Source/Editor/ContentBrowserData/Public/ContentBrowserDataSubsystem.h
+++ b/Engine/Source/Editor/ContentBrowserData/Public/ContentBrowserDataSubsystem.h
@@ -275,9 +275,16 @@ public:
 	 * Query whether the given virtual folder should be visible in the UI.
 	 * @param Path Virtual path of folder e.g. /All/GameData/Stuff
 	 * @param Flags Currently unused
-	 * @param ContentsFilter Optional filter for limiting visibility to only folders with a certain type of contents based on desired view filtering.
 	 */
-	bool IsFolderVisible(const FName Path, const EContentBrowserIsFolderVisibleFlags Flags = EContentBrowserIsFolderVisibleFlags::Default, const FContentBrowserFolderContentsFilter& ContentsFilter = {}) const;
+	bool IsFolderVisible(const FName Path, const EContentBrowserIsFolderVisibleFlags Flags = EContentBrowserIsFolderVisibleFlags::Default) const;
+
+	/**
+	 * Query whether the given virtual folder should be visible in the UI.
+	 * @param Path Virtual path of folder e.g. /All/GameData/Stuff
+	 * @param Flags Currently unused
+	 * @param ContentsFilter Filter for limiting visibility to only folders with a certain type of contents based on desired view filtering.
+	 */
+	bool IsFolderVisible(const FName Path, const EContentBrowserIsFolderVisibleFlags Flags, const FContentBrowserFolderContentsFilter& ContentsFilter) const;
 
 	/**
 	 * Query whether the given virtual folder should be visible if the UI is asking to hide empty content folders.
@@ -304,6 +311,16 @@ public:
 	 */
 	FContentBrowserItemTemporaryContext CreateFolder(const FName InPath) const;
 
+	/*
+	 * Attempt to begin the process of asynchronously creating a folder at the given virtual path, returning a temporary item that can be finalized or canceled by the user.
+	 *
+	 * @param InPath The initial virtual path of the folder that is being created.
+	 * @param HideFolderIfEmptyFilter Current empty folder filter used to permit new folder operations on hidden folders.
+	 *
+	 * @return The pending folder item to create (test for validity).
+	 */
+	FContentBrowserItemTemporaryContext CreateFolder(const FName InPath, const TSharedPtr<IContentBrowserHideFolderIfEmptyFilter>& HideFolderIfEmptyFilter) const;
+
 	/**
 	 * Attempt to convert the given package path to virtual paths associated with the active data sources (callback will be called for each successful conversion).
 	 * @note This exists to allow the Content Browser to interface with public APIs that only operate on package paths and should ideally be avoided for new code.
diff --git a/Engine/Source/Editor/ContentBrowserData/Public/ContentBrowserItem.h b/Engine/Source/Editor/ContentBrowserData/Public/ContentBrowserItem.h
--- a/Engine/Source/Editor/ContentBrowserData/Public/ContentBrowserItem.h
+++ b/Engine/Source/Editor/ContentBrowserData/Public/ContentBrowserItem.h
@@ -21,6 +21,7 @@
 
 class FAssetThumbnail;
 class FContentBrowserMinimalItemData;
+class IContentBrowserHideFolderIfEmptyFilter;
 class UContentBrowserDataSource;
 class UObject;
 struct FAssetData;
@@ -349,6 +350,17 @@ public:
 	 */
 	bool CanRename(const FString* InNewName, FText* OutErrorMsg = nullptr) const;
 
+	/**
+	 * Query whether this item is can be renamed, optionally providing error information if it cannot.
+	 *
+	 * @param InNewName Optional name that will be tested for validity alongside the general ability to rename this item.
+	 * @param HideFolderIfEmptyFilter Current empty folder filter used to permit rename operations to hidden folders.
+	 * @param OutErrorMsg Optional error message to fill on failure.
+	 *
+	 * @return True if the item can be renamed, false otherwise.
+	 */
+	bool CanRename(const FString* InNewName, const IContentBrowserHideFolderIfEmptyFilter* HideFolderIfEmptyFilter, FText* OutErrorMsg = nullptr) const;
+
 	/**
 	 * Attempt to rename this item.
 	 *
