* Fixed ensure when trying to capture thumbnail on an unloaded asset : we now try load the package (like is done when clearing the thumbnail)
* Also changed label "Capture Thumbnail" into "Capture Thumbnail From Viewport", to make it explicit that this will _not_ capture the asset scene's thumbnail, but rather, use the current viewport as its thumbnail (the tooltip was already explicit about it)

diff --git a/Engine/Plugins/Editor/ContentBrowser/ContentBrowserAssetDataSource/Source/ContentBrowserAssetDataSource/Private/AssetFileContextMenu.cpp b/Engine/Plugins/Editor/ContentBrowser/ContentBrowserAssetDataSource/Source/ContentBrowserAssetDataSource/Private/AssetFileContextMenu.cpp
--- a/Engine/Plugins/Editor/ContentBrowser/ContentBrowserAssetDataSource/Source/ContentBrowserAssetDataSource/Private/AssetFileContextMenu.cpp
+++ b/Engine/Plugins/Editor/ContentBrowser/ContentBrowserAssetDataSource/Source/ContentBrowserAssetDataSource/Private/AssetFileContextMenu.cpp
@@ -337,7 +337,7 @@ void FAssetFileContextMenu::MakeAssetActionsSubMenu(UToolMenu* Menu)
 		{
 			Section.AddMenuEntry(
 				"CaptureThumbnail",
-				LOCTEXT("CaptureThumbnail", "Capture Thumbnail"),
+				LOCTEXT("CaptureThumbnail", "Capture Thumbnail From Viewport"),
 				LOCTEXT("CaptureThumbnailTooltip", "Captures a thumbnail from the active viewport."),
 				FSlateIcon(),
 				FUIAction(
diff --git a/Engine/Source/Developer/AssetTools/Private/AssetViewUtils.cpp b/Engine/Source/Developer/AssetTools/Private/AssetViewUtils.cpp
--- a/Engine/Source/Developer/AssetTools/Private/AssetViewUtils.cpp
+++ b/Engine/Source/Developer/AssetTools/Private/AssetViewUtils.cpp
@@ -965,7 +965,7 @@ void AssetViewUtils::CaptureThumbnailFromViewport(FViewport* InViewport, const T
 			const FString ObjectFullName = CurrentAsset.GetFullName();
 			const FString PackageName    = CurrentAsset.PackageName.ToString();
 
-			UPackage* AssetPackage = FindObject<UPackage>( NULL, *PackageName );
+			UPackage* AssetPackage = CurrentAsset.GetPackage();
 			if ( ensure(AssetPackage) )
 			{
 				FObjectThumbnail* NewThumbnail = ThumbnailTools::CacheThumbnail(ObjectFullName, &TempThumbnail, AssetPackage);
