Backout the menu label name change ("Capture Thumbnail" -> "Capture Thumbnail From Viewport") of 39590014

diff --git a/Engine/Plugins/Editor/ContentBrowser/ContentBrowserAssetDataSource/Source/ContentBrowserAssetDataSource/Private/AssetFileContextMenu.cpp b/Engine/Plugins/Editor/ContentBrowser/ContentBrowserAssetDataSource/Source/ContentBrowserAssetDataSource/Private/AssetFileContextMenu.cpp
--- a/Engine/Plugins/Editor/ContentBrowser/ContentBrowserAssetDataSource/Source/ContentBrowserAssetDataSource/Private/AssetFileContextMenu.cpp
+++ b/Engine/Plugins/Editor/ContentBrowser/ContentBrowserAssetDataSource/Source/ContentBrowserAssetDataSource/Private/AssetFileContextMenu.cpp
@@ -337,7 +337,7 @@ void FAssetFileContextMenu::MakeAssetActionsSubMenu(UToolMenu* Menu)
 		{
 			Section.AddMenuEntry(
 				"CaptureThumbnail",
-				LOCTEXT("CaptureThumbnail", "Capture Thumbnail From Viewport"),
+				LOCTEXT("CaptureThumbnail", "Capture Thumbnail"),
 				LOCTEXT("CaptureThumbnailTooltip", "Captures a thumbnail from the active viewport."),
 				FSlateIcon(),
 				FUIAction(
