[ContentBrowser] CB right click menu Import Entry should just say "Import to Current Folder"

diff --git a/Engine/Plugins/Editor/ContentBrowser/ContentBrowserAssetDataSource/Source/ContentBrowserAssetDataSource/Private/NewAssetContextMenu.cpp b/Engine/Plugins/Editor/ContentBrowser/ContentBrowserAssetDataSource/Source/ContentBrowserAssetDataSource/Private/NewAssetContextMenu.cpp
--- a/Engine/Plugins/Editor/ContentBrowser/ContentBrowserAssetDataSource/Source/ContentBrowserAssetDataSource/Private/NewAssetContextMenu.cpp
+++ b/Engine/Plugins/Editor/ContentBrowser/ContentBrowserAssetDataSource/Source/ContentBrowserAssetDataSource/Private/NewAssetContextMenu.cpp
@@ -247,7 +247,7 @@ void FNewAssetContextMenu::MakeContextMenu(
 			FToolMenuSection& Section = Menu->FindOrAddSection("ContentBrowserGetContent");
 			Section.AddMenuEntry(
 				"ImportAsset",
-				FText::Format(LOCTEXT("ImportAsset", "Import to {0}..."), FText::FromName(FirstSelectedPath)),
+				LOCTEXT("ImportAsset", "Import to Current Folder"),
 				LOCTEXT("ImportAssetTooltip_NewAsset", "Imports an asset from file to this folder."),
 				FSlateIcon(FAppStyle::GetAppStyleSetName(), "Icons.Import"),
 				FUIAction(
