Fix thumbnail rendering forcing textures to become fully resident. This can unnecessarily fill the streaming pools over time.
Instead we set a time period for them to be fully resident. This is the same way that some other systems work.

diff --git a/Engine/Plugins/Editor/ContentBrowser/ContentBrowserAssetDataSource/Source/ContentBrowserAssetDataSource/Private/ContentBrowserAssetDataPayload.cpp b/Engine/Plugins/Editor/ContentBrowser/ContentBrowserAssetDataSource/Source/ContentBrowserAssetDataSource/Private/ContentBrowserAssetDataPayload.cpp
--- a/Engine/Plugins/Editor/ContentBrowser/ContentBrowserAssetDataSource/Source/ContentBrowserAssetDataSource/Private/ContentBrowserAssetDataPayload.cpp
+++ b/Engine/Plugins/Editor/ContentBrowser/ContentBrowserAssetDataSource/Source/ContentBrowserAssetDataSource/Private/ContentBrowserAssetDataPayload.cpp
@@ -160,13 +160,14 @@ void FContentBrowserAssetFileItemDataPayload::UpdateThumbnail(FAssetThumbnail& I
 {
 	if (UObject* AssetPtr = GetAsset())
 	{
+		constexpr float TimeToBoostTextureResidency = 10.0f;
 		if (UTexture2D* TextureAssetPtr = Cast<UTexture2D>(AssetPtr))
 		{
-			TextureAssetPtr->bForceMiplevelsToBeResident = true;
+			TextureAssetPtr->SetForceMipLevelsToBeResident(TimeToBoostTextureResidency);
 		}
 		else if (UMaterial* MaterialAssetPtr = Cast<UMaterial>(AssetPtr))
 		{
-			MaterialAssetPtr->SetForceMipLevelsToBeResident(true, true, -1.0f);
+			MaterialAssetPtr->SetForceMipLevelsToBeResident(false, false, TimeToBoostTextureResidency);
 		}
 	}
 
