Updates ContentBrowserData to support retrieving an item's EAssetAccessSpecifier.

#rnx

diff --git a/Engine/Plugins/Editor/ContentBrowser/ContentBrowserAliasDataSource/Source/ContentBrowserAliasDataSource/Private/ContentBrowserAliasDataSource.cpp b/Engine/Plugins/Editor/ContentBrowser/ContentBrowserAliasDataSource/Source/ContentBrowserAliasDataSource/Private/ContentBrowserAliasDataSource.cpp
--- a/Engine/Plugins/Editor/ContentBrowser/ContentBrowserAliasDataSource/Source/ContentBrowserAliasDataSource/Private/ContentBrowserAliasDataSource.cpp
+++ b/Engine/Plugins/Editor/ContentBrowser/ContentBrowserAliasDataSource/Source/ContentBrowserAliasDataSource/Private/ContentBrowserAliasDataSource.cpp
@@ -997,6 +997,16 @@ bool UContentBrowserAliasDataSource::TryGetCollectionId(const FContentBrowserIte
 	return false;
 }
 
+bool UContentBrowserAliasDataSource::GetItemAssetAccessSpecifier(const FContentBrowserItemData& InItem, EAssetAccessSpecifier& OutAssetAccessSpecifier)
+{
+	return ContentBrowserAssetData::GetAssetAccessSpecifier(this, InItem, OutAssetAccessSpecifier);
+}
+
+bool UContentBrowserAliasDataSource::CanModifyItemAssetAccessSpecifier(const FContentBrowserItemData& InItem)
+{
+	return ContentBrowserAssetData::CanModifyAssetAccessSpecifier(AssetTools, this, InItem);
+}
+
 bool UContentBrowserAliasDataSource::Legacy_TryGetPackagePath(const FContentBrowserItemData& InItem, FName& OutPackagePath)
 {
 	if (TSharedPtr<const FContentBrowserAssetFolderItemDataPayload> FolderPayload = ContentBrowserAssetData::GetAssetFolderItemPayload(this, InItem))
diff --git a/Engine/Plugins/Editor/ContentBrowser/ContentBrowserAliasDataSource/Source/ContentBrowserAliasDataSource/Public/ContentBrowserAliasDataSource.h b/Engine/Plugins/Editor/ContentBrowser/ContentBrowserAliasDataSource/Source/ContentBrowserAliasDataSource/Public/ContentBrowserAliasDataSource.h
--- a/Engine/Plugins/Editor/ContentBrowser/ContentBrowserAliasDataSource/Source/ContentBrowserAliasDataSource/Public/ContentBrowserAliasDataSource.h
+++ b/Engine/Plugins/Editor/ContentBrowser/ContentBrowserAliasDataSource/Source/ContentBrowserAliasDataSource/Public/ContentBrowserAliasDataSource.h
@@ -100,6 +100,9 @@ public:
 
 	virtual bool TryGetCollectionId(const FContentBrowserItemData& InItem, FSoftObjectPath& OutCollectionId) override;
 
+	virtual bool GetItemAssetAccessSpecifier(const FContentBrowserItemData& InItem, EAssetAccessSpecifier& OutAssetAccessSpecifier) override;
+	virtual bool CanModifyItemAssetAccessSpecifier(const FContentBrowserItemData& InItem) override;
+
 	// Legacy functions seem necessary for FrontendFilters to work
 	virtual bool Legacy_TryGetPackagePath(const FContentBrowserItemData& InItem, FName& OutPackagePath) override;
 	virtual bool Legacy_TryGetAssetData(const FContentBrowserItemData& InItem, FAssetData& OutAssetData) override;
diff --git a/Engine/Plugins/Editor/ContentBrowser/ContentBrowserAssetDataSource/Source/ContentBrowserAssetDataSource/Private/ContentBrowserAssetDataCore.cpp b/Engine/Plugins/Editor/ContentBrowser/ContentBrowserAssetDataSource/Source/ContentBrowserAssetDataSource/Private/ContentBrowserAssetDataCore.cpp
--- a/Engine/Plugins/Editor/ContentBrowser/ContentBrowserAssetDataSource/Source/ContentBrowserAssetDataSource/Private/ContentBrowserAssetDataCore.cpp
+++ b/Engine/Plugins/Editor/ContentBrowser/ContentBrowserAssetDataSource/Source/ContentBrowserAssetDataSource/Private/ContentBrowserAssetDataCore.cpp
@@ -2061,6 +2061,32 @@ bool GetAssetDataAttributes(const FAssetData& InAssetData, const bool InIncludeM
 	return true;
 }
 
+bool GetAssetAccessSpecifier(const UContentBrowserDataSource* InOwnerDataSource, const FContentBrowserItemData& InItem, EAssetAccessSpecifier& OutAssetAccessSpecifier)
+{
+	if (TSharedPtr<const FContentBrowserAssetFileItemDataPayload> AssetPayload = GetAssetFileItemPayload(InOwnerDataSource, InItem))
+	{
+		OutAssetAccessSpecifier = AssetPayload->GetAssetData().GetAssetAccessSpecifier();
+		return true;
+	}
+
+	return false;
+}
+
+bool CanModifyAssetAccessSpecifier(IAssetTools* InAssetTools, const UContentBrowserDataSource* InOwnerDataSource, const FContentBrowserItemData& InItem)
+{
+	check(InAssetTools);
+
+	if (TSharedPtr<const FContentBrowserAssetFileItemDataPayload> AssetPayload = GetAssetFileItemPayload(InOwnerDataSource, InItem))
+	{
+		FNameBuilder AssetPath;
+		FTopLevelAssetPath(AssetPayload->GetAssetData().PackageName, AssetPayload->GetAssetData().AssetName).AppendString(AssetPath);
+
+		return InAssetTools->CanAssetBePublic(AssetPath);
+	}
+
+	return false;
+}
+
 void PopulateAssetFolderContextMenu(UContentBrowserDataSource* InOwnerDataSource, UToolMenu* InMenu, FAssetFolderContextMenu& InAssetFolderContextMenu)
 {
 	const UContentBrowserDataMenuContext_FolderMenu* ContextObject = InMenu->FindContext<UContentBrowserDataMenuContext_FolderMenu>();
diff --git a/Engine/Plugins/Editor/ContentBrowser/ContentBrowserAssetDataSource/Source/ContentBrowserAssetDataSource/Private/ContentBrowserAssetDataSource.cpp b/Engine/Plugins/Editor/ContentBrowser/ContentBrowserAssetDataSource/Source/ContentBrowserAssetDataSource/Private/ContentBrowserAssetDataSource.cpp
--- a/Engine/Plugins/Editor/ContentBrowser/ContentBrowserAssetDataSource/Source/ContentBrowserAssetDataSource/Private/ContentBrowserAssetDataSource.cpp
+++ b/Engine/Plugins/Editor/ContentBrowser/ContentBrowserAssetDataSource/Source/ContentBrowserAssetDataSource/Private/ContentBrowserAssetDataSource.cpp
@@ -2642,6 +2642,16 @@ bool UContentBrowserAssetDataSource::GetItemPhysicalPath(const FContentBrowserIt
 	return ContentBrowserAssetData::GetItemPhysicalPath(this, InItem, OutDiskPath);
 }
 
+bool UContentBrowserAssetDataSource::GetItemAssetAccessSpecifier(const FContentBrowserItemData& InItem, EAssetAccessSpecifier& OutAssetAccessSpecifier)
+{
+	return ContentBrowserAssetData::GetAssetAccessSpecifier(this, InItem, OutAssetAccessSpecifier);
+}
+
+bool UContentBrowserAssetDataSource::CanModifyItemAssetAccessSpecifier(const FContentBrowserItemData& InItem)
+{
+	return ContentBrowserAssetData::CanModifyAssetAccessSpecifier(AssetTools, this, InItem);
+}
+
 bool UContentBrowserAssetDataSource::IsItemDirty(const FContentBrowserItemData& InItem)
 {
 	return ContentBrowserAssetData::IsItemDirty(this, InItem);
diff --git a/Engine/Plugins/Editor/ContentBrowser/ContentBrowserAssetDataSource/Source/ContentBrowserAssetDataSource/Public/ContentBrowserAssetDataCore.h b/Engine/Plugins/Editor/ContentBrowser/ContentBrowserAssetDataSource/Source/ContentBrowserAssetDataSource/Public/ContentBrowserAssetDataCore.h
--- a/Engine/Plugins/Editor/ContentBrowser/ContentBrowserAssetDataSource/Source/ContentBrowserAssetDataSource/Public/ContentBrowserAssetDataCore.h
+++ b/Engine/Plugins/Editor/ContentBrowser/ContentBrowserAssetDataSource/Source/ContentBrowserAssetDataSource/Public/ContentBrowserAssetDataCore.h
@@ -218,6 +218,10 @@ CONTENTBROWSERASSETDATASOURCE_API FContentBrowserItemData CreateUnsupportedAsset
 
 	CONTENTBROWSERASSETDATASOURCE_API bool GetAssetDataAttributes(const FAssetData& InAssetData, const bool InIncludeMetaData, FContentBrowserItemDataAttributeValues& OutAttributeValues);
 
+	CONTENTBROWSERASSETDATASOURCE_API bool GetAssetAccessSpecifier(const UContentBrowserDataSource* InOwnerDataSource, const FContentBrowserItemData& InItem, EAssetAccessSpecifier& OutAssetAccessSpecifier);
+
+	CONTENTBROWSERASSETDATASOURCE_API bool CanModifyAssetAccessSpecifier(IAssetTools* InAssetTools, const UContentBrowserDataSource* InOwnerDataSource, const FContentBrowserItemData& InItem);
+
 	CONTENTBROWSERASSETDATASOURCE_API void PopulateAssetFolderContextMenu(UContentBrowserDataSource* InOwnerDataSource, UToolMenu* InMenu, FAssetFolderContextMenu& InAssetFolderContextMenu);
 
 	CONTENTBROWSERASSETDATASOURCE_API void PopulateAssetFileContextMenu(UContentBrowserDataSource* InOwnerDataSource, UToolMenu* InMenu, FAssetFileContextMenu& InAssetFileContextMenu);
diff --git a/Engine/Plugins/Editor/ContentBrowser/ContentBrowserAssetDataSource/Source/ContentBrowserAssetDataSource/Public/ContentBrowserAssetDataSource.h b/Engine/Plugins/Editor/ContentBrowser/ContentBrowserAssetDataSource/Source/ContentBrowserAssetDataSource/Public/ContentBrowserAssetDataSource.h
--- a/Engine/Plugins/Editor/ContentBrowser/ContentBrowserAssetDataSource/Source/ContentBrowserAssetDataSource/Public/ContentBrowserAssetDataSource.h
+++ b/Engine/Plugins/Editor/ContentBrowser/ContentBrowserAssetDataSource/Source/ContentBrowserAssetDataSource/Public/ContentBrowserAssetDataSource.h
@@ -267,6 +267,10 @@ public:
 
 	virtual bool GetItemPhysicalPath(const FContentBrowserItemData& InItem, FString& OutDiskPath) override;
 
+	virtual bool GetItemAssetAccessSpecifier(const FContentBrowserItemData& InItem, EAssetAccessSpecifier& OutAssetAccessSpecifier) override;
+
+	virtual bool CanModifyItemAssetAccessSpecifier(const FContentBrowserItemData& InItem) override;
+
 	virtual bool IsItemDirty(const FContentBrowserItemData& InItem) override;
 
 	virtual bool CanEditItem(const FContentBrowserItemData& InItem, FText* OutErrorMsg) override;
diff --git a/Engine/Source/Editor/ContentBrowserData/Private/ContentBrowserDataSource.cpp b/Engine/Source/Editor/ContentBrowserData/Private/ContentBrowserDataSource.cpp
--- a/Engine/Source/Editor/ContentBrowserData/Private/ContentBrowserDataSource.cpp
+++ b/Engine/Source/Editor/ContentBrowserData/Private/ContentBrowserDataSource.cpp
@@ -303,6 +303,16 @@ bool UContentBrowserDataSource::GetItemPhysicalPath(const FContentBrowserItemDat
 	return false;
 }
 
+bool UContentBrowserDataSource::GetItemAssetAccessSpecifier(const FContentBrowserItemData& InItem, EAssetAccessSpecifier& OutAssetAccessSpecifier)
+{
+	return false;
+}
+
+bool UContentBrowserDataSource::CanModifyItemAssetAccessSpecifier(const FContentBrowserItemData& InItem)
+{
+	return false;
+}
+
 bool UContentBrowserDataSource::IsItemDirty(const FContentBrowserItemData& InItem)
 {
 	return false;
diff --git a/Engine/Source/Editor/ContentBrowserData/Private/ContentBrowserItem.cpp b/Engine/Source/Editor/ContentBrowserData/Private/ContentBrowserItem.cpp
--- a/Engine/Source/Editor/ContentBrowserData/Private/ContentBrowserItem.cpp
+++ b/Engine/Source/Editor/ContentBrowserData/Private/ContentBrowserItem.cpp
@@ -413,6 +413,22 @@ FText FContentBrowserItem::GetDisplayName() const
 		: FText();
 }
 
+bool FContentBrowserItem::GetItemAssetAccessSpecifier(EAssetAccessSpecifier& OutAssetAccessSpecifier) const
+{
+	const FContentBrowserItemData* PrimaryItemData = GetPrimaryInternalItem();
+	return PrimaryItemData
+		? PrimaryItemData->GetItemAssetAccessSpecifier(OutAssetAccessSpecifier)
+		: false;
+}
+
+bool FContentBrowserItem::CanModifyItemAssetAccessSpecifier() const
+{
+	const FContentBrowserItemData* PrimaryItemData = GetPrimaryInternalItem();
+	return PrimaryItemData
+		? PrimaryItemData->CanModifyItemAssetAccessSpecifier()
+		: false;
+}
+
 struct FContentBrowserItemHelper
 {
 public:
diff --git a/Engine/Source/Editor/ContentBrowserData/Private/ContentBrowserItemData.cpp b/Engine/Source/Editor/ContentBrowserData/Private/ContentBrowserItemData.cpp
--- a/Engine/Source/Editor/ContentBrowserData/Private/ContentBrowserItemData.cpp
+++ b/Engine/Source/Editor/ContentBrowserData/Private/ContentBrowserItemData.cpp
@@ -145,6 +145,26 @@ FText FContentBrowserItemData::GetDisplayName() const
 	return CachedDisplayName;
 }
 
+bool FContentBrowserItemData::GetItemAssetAccessSpecifier(EAssetAccessSpecifier& OutAssetAccessSpecifier) const
+{
+	if (UContentBrowserDataSource* ItemDataSource = GetOwnerDataSource())
+	{
+		return ItemDataSource->GetItemAssetAccessSpecifier(*this, OutAssetAccessSpecifier);
+	}
+
+	return false;
+}
+
+bool FContentBrowserItemData::CanModifyItemAssetAccessSpecifier() const
+{
+	if (UContentBrowserDataSource* ItemDataSource = GetOwnerDataSource())
+	{
+		return ItemDataSource->CanModifyItemAssetAccessSpecifier(*this);
+	}
+
+	return false;
+}
+
 TSharedPtr<const IContentBrowserItemDataPayload> FContentBrowserItemData::GetPayload() const
 {
 	return Payload;
diff --git a/Engine/Source/Editor/ContentBrowserData/Public/ContentBrowserDataSource.h b/Engine/Source/Editor/ContentBrowserData/Public/ContentBrowserDataSource.h
--- a/Engine/Source/Editor/ContentBrowserData/Public/ContentBrowserDataSource.h
+++ b/Engine/Source/Editor/ContentBrowserData/Public/ContentBrowserDataSource.h
@@ -345,6 +345,25 @@ public:
 	 */
 	virtual bool GetItemPhysicalPath(const FContentBrowserItemData& InItem, FString& OutDiskPath);
 
+	/**
+	 * Query the EAssetAccessSpecifier of the item
+	 *
+	 * @param InItem The item to query.
+	 * @param OutAssetAccessSpecifier The EAssetAccessSpecifier reference that will be filled.
+	 *
+	 * @return True if OutAssetAccessSpecifier was filled.
+	 */
+	virtual bool GetItemAssetAccessSpecifier(const FContentBrowserItemData& InItem, EAssetAccessSpecifier& OutAssetAccessSpecifier);
+
+	/**
+	 * Query the item's ability to modify its EAssetAccessSpecifier
+	 *
+	 * @param InItem The item to query.
+	 *
+	 * @return True if the item has an EAssetAccessSpecifier that can be modified.
+	 */
+	virtual bool CanModifyItemAssetAccessSpecifier(const FContentBrowserItemData& InItem);
+
 	/**
 	 * Query whether the given item is considered dirty (ie, has unsaved changes).
 	 *
diff --git a/Engine/Source/Editor/ContentBrowserData/Public/ContentBrowserItem.h b/Engine/Source/Editor/ContentBrowserData/Public/ContentBrowserItem.h
--- a/Engine/Source/Editor/ContentBrowserData/Public/ContentBrowserItem.h
+++ b/Engine/Source/Editor/ContentBrowserData/Public/ContentBrowserItem.h
@@ -19,6 +19,7 @@
 
 #include "ContentBrowserItem.generated.h"
 
+enum class EAssetAccessSpecifier : uint8;
 struct FContentBrowserItemPath;
 class FAssetThumbnail;
 class FContentBrowserMinimalItemData;
@@ -217,6 +218,22 @@ public:
 	 */
 	FText GetDisplayName() const;
 
+	/**
+	 * Query the EAssetAccessSpecifier of the item.
+	 * 
+	 * @param OutAssetAccessSpecifier The EAssetAccessSpecifier reference that will be filled.
+	 *
+	 * @return True if OutAssetAccessSpecifier was filled.
+	 */
+	bool GetItemAssetAccessSpecifier(EAssetAccessSpecifier& OutAssetAccessSpecifier) const;
+
+	/**
+	 * Query the item's ability to modify its EAssetAccessSpecifier.
+	 *
+	 * @return True if the item has an EAssetAccessSpecifier that can be modified.
+	 */
+	bool CanModifyItemAssetAccessSpecifier() const;
+
 	/**
 	 * Query the value of the given attribute on this item.
 	 *
diff --git a/Engine/Source/Editor/ContentBrowserData/Public/ContentBrowserItemData.h b/Engine/Source/Editor/ContentBrowserData/Public/ContentBrowserItemData.h
--- a/Engine/Source/Editor/ContentBrowserData/Public/ContentBrowserItemData.h
+++ b/Engine/Source/Editor/ContentBrowserData/Public/ContentBrowserItemData.h
@@ -19,6 +19,7 @@
 #include "UObject/WeakObjectPtr.h"
 #include "UObject/WeakObjectPtrTemplates.h"
 
+enum class EAssetAccessSpecifier : uint8;
 class FContentBrowserMinimalItemData;
 class IContentBrowserItemDataSink;
 class UContentBrowserDataSource;
@@ -261,6 +262,22 @@ public:
 	 */
 	FText GetDisplayName() const;
 
+	/**
+	 * Query the EAssetAccessSpecifier of the item.
+	 *
+	 * @param OutAssetAccessSpecifier The EAssetAccessSpecifier reference that will be filled.
+	 *
+	 * @return True if OutAssetAccessSpecifier was filled.
+	 */
+	bool GetItemAssetAccessSpecifier(EAssetAccessSpecifier& OutAssetAccessSpecifier) const;
+
+	/**
+	 * Query the item's ability to modify its EAssetAccessSpecifier.
+	 *
+	 * @return True if the item has an EAssetAccessSpecifier that can be modified.
+	 */
+	bool CanModifyItemAssetAccessSpecifier() const;
+
 	/**
 	 * Get any data source defined payload data for this item.
 	 */
