- Added a "Reload Uncooked Asset" menu option to the Asset Actions right-click menu in the content browser
- Added a way for an asset to request other assets to be uncooked as well, but nothing is using it yet (marked as UE_INTERNAL)

diff --git a/Engine/Plugins/Editor/ContentBrowser/ContentBrowserAssetDataSource/Source/ContentBrowserAssetDataSource/Private/AssetFileContextMenu.cpp b/Engine/Plugins/Editor/ContentBrowser/ContentBrowserAssetDataSource/Source/ContentBrowserAssetDataSource/Private/AssetFileContextMenu.cpp
--- a/Engine/Plugins/Editor/ContentBrowser/ContentBrowserAssetDataSource/Source/ContentBrowserAssetDataSource/Private/AssetFileContextMenu.cpp
+++ b/Engine/Plugins/Editor/ContentBrowser/ContentBrowserAssetDataSource/Source/ContentBrowserAssetDataSource/Private/AssetFileContextMenu.cpp
@@ -72,6 +72,9 @@
 #include "UnrealEdGlobals.h"
 #include "Algo/AnyOf.h"
 #include "Misc/WarnIfAssetsLoadedInScope.h"
+#include "CoreGlobalsInternal.h"
+#include "Serialization/PackageStore.h"
+
 
 #define LOCTEXT_NAMESPACE "ContentBrowser"
 
@@ -86,6 +89,11 @@ void FAssetFileContextMenu::MakeContextMenu(UToolMenu* InMenu, const TArray<FAss
 	{
 		return;
 	}
+
+	if (Context)
+	{
+		OnRefreshView = Context->OnRefreshView;
+	}
 	
 	if (SelectedAssets.Num() > 0)
 	{
@@ -432,6 +440,22 @@ void FAssetFileContextMenu::MakeAssetActionsSubMenu(UToolMenu* Menu)
 			)
 		);
 
+		if (IsRunningHybridCookedEditor())
+		{
+			// Reload
+			Section.AddMenuEntry(
+				"ReloadUncooked",
+				LOCTEXT("ReloadUncooked", "Reload Uncooked Version"),
+				LOCTEXT("ReloadTooltip", "Reload the selected assets from their uncooked file on disk."),
+				FSlateIcon(),
+				FUIAction(
+					FExecuteAction::CreateSP(this, &FAssetFileContextMenu::ExecuteReloadUncooked),
+					FCanExecuteAction::CreateSP(this, &FAssetFileContextMenu::CanExecuteReloadUncooked)
+				)
+			);
+		}
+
+
 		// Load
 		Section.AddMenuEntry(
 			"Load",
@@ -1624,7 +1648,34 @@ bool FAssetFileContextMenu::CanExecuteReload() const
 	return SelectedAssets.Num() > 0;
 }
 
-void FAssetFileContextMenu::ExecuteReload()
+static void ProcessObjectForUncook(const UObject* Obj, TArray<UPackage*>& PackagesToReload)
+{
+	if (Obj == nullptr)
+	{
+		return;
+	}
+
+	// tell package stores to stop serving cooked packages
+	FHybridPackageStoreBackend::ForceLoadPackageAsLoose(Obj->GetPackage()->GetPackageId());
+
+	UPackage* Pkg = Obj->GetPackage();
+	if (Pkg->HasAnyPackageFlags(PKG_Cooked))
+	{
+		PackagesToReload.AddUnique(Pkg);
+	}
+
+	// now unload any additional packages for anything this object needs
+	TArray<UObject*> AdditionalObjects;
+	PRAGMA_DISABLE_INTERNAL_WARNINGS
+	Obj->GetAdditionalObjectsToUncook(AdditionalObjects);
+	PRAGMA_ENABLE_INTERNAL_WARNINGS
+	for (UObject* RelatedObj : AdditionalObjects)
+	{
+		ProcessObjectForUncook(RelatedObj, PackagesToReload);
+	}
+}
+
+static void ReloadPackages(TArray<FAssetData>& Assets, bool bUncookPackages)
 {
 	// Don't allow asset reload during PIE
 	if (GIsEditor)
@@ -1640,11 +1691,11 @@ void FAssetFileContextMenu::ExecuteReload()
 		}
 	}
 
-	if (SelectedAssets.Num() > 0)
+	if (Assets.Num() > 0)
 	{
 		TArray<UPackage*> PackagesToReload;
 
-		for (auto AssetIt = SelectedAssets.CreateConstIterator(); AssetIt; ++AssetIt)
+		for (auto AssetIt = Assets.CreateConstIterator(); AssetIt; ++AssetIt)
 		{
 			const FAssetData& AssetData = *AssetIt;
 
@@ -1670,7 +1721,24 @@ void FAssetFileContextMenu::ExecuteReload()
 				continue;
 			}
 
-			PackagesToReload.AddUnique(AssetData.GetPackage());
+			// simple case of just reloading
+			if (!bUncookPackages)
+			{
+				PackagesToReload.AddUnique(AssetData.GetPackage());
+			}
+			// complexity for uncooked a packages and it's dependencies
+			else
+			{
+				// if the object wasn't loaded, maybe it's because it was cooked, so reload the whole package uncooked,
+				// then we can try again for it's related objects
+				UObject* AssetObj = AssetData.FastGetAsset();
+				if (AssetObj == nullptr)
+				{
+					UPackageTools::ReloadPackages({ AssetData.GetPackage() } );
+				}
+
+				ProcessObjectForUncook(AssetData.FastGetAsset(), PackagesToReload);
+			}
 		}
 
 		if (PackagesToReload.Num() > 0)
@@ -1680,6 +1748,23 @@ void FAssetFileContextMenu::ExecuteReload()
 	}
 }
 
+void FAssetFileContextMenu::ExecuteReload()
+{
+	ReloadPackages(SelectedAssets, /*bUncook*/false);
+}
+
+bool FAssetFileContextMenu::CanExecuteReloadUncooked() const
+{
+	return SelectedAssets.Num() > 0;
+}
+
+void FAssetFileContextMenu::ExecuteReloadUncooked()
+{
+	ReloadPackages(SelectedAssets, /*bUncook*/true);
+	OnRefreshView.ExecuteIfBound();
+}
+
+
 void FAssetFileContextMenu::ExecuteConsolidate()
 {
 	TArray<UObject*> ObjectsToConsolidate;
diff --git a/Engine/Plugins/Editor/ContentBrowser/ContentBrowserAssetDataSource/Source/ContentBrowserAssetDataSource/Public/AssetFileContextMenu.h b/Engine/Plugins/Editor/ContentBrowser/ContentBrowserAssetDataSource/Source/ContentBrowserAssetDataSource/Public/AssetFileContextMenu.h
--- a/Engine/Plugins/Editor/ContentBrowser/ContentBrowserAssetDataSource/Source/ContentBrowserAssetDataSource/Public/AssetFileContextMenu.h
+++ b/Engine/Plugins/Editor/ContentBrowser/ContentBrowserAssetDataSource/Source/ContentBrowserAssetDataSource/Public/AssetFileContextMenu.h
@@ -16,6 +16,7 @@ class FAssetFileContextMenu : public TSharedFromThis<FAssetFileContextMenu>
 {
 public:
 	DECLARE_DELEGATE_OneParam(FOnShowAssetsInPathsView, const TArray<FAssetData>& /*AssetsToFind*/);
+	DECLARE_DELEGATE(FOnRefreshView);
 
 	/** Makes the context menu widget */
 	UE_API void MakeContextMenu(
@@ -77,6 +78,12 @@ private:
 	/** Handler for Reload */
 	UE_API void ExecuteReload();
 
+	/** Handler to check to see if a reload uncooked command is allowed */
+	bool CanExecuteReloadUncooked() const;
+
+	/** Handler for Reload Uncooked */
+	void ExecuteReloadUncooked();
+
 	/** Is allowed to modify files or folders under this path */
 	UE_API bool CanModifyPath(const FString& InPath) const;
 
@@ -244,6 +251,7 @@ private:
 	TWeakPtr<SWidget> ParentWidget;
 
 	FOnShowAssetsInPathsView OnShowAssetsInPathsView;
+	FOnRefreshView OnRefreshView;
 
 	/** Cached CanExecute vars */
 	bool bAtLeastOneNonRedirectorSelected = false;
diff --git a/Engine/Source/Runtime/CoreUObject/Public/UObject/Object.h b/Engine/Source/Runtime/CoreUObject/Public/UObject/Object.h
--- a/Engine/Source/Runtime/CoreUObject/Public/UObject/Object.h
+++ b/Engine/Source/Runtime/CoreUObject/Public/UObject/Object.h
@@ -315,6 +315,11 @@ public:
 
 	/** Utility to allow overrides of Modify to avoid doing work if this object cannot be safely modified */
 	COREUOBJECT_API bool CanModify() const;
+
+	/** Returns a list of additional objects that will need to be replaced with uncooked versions when this object is being replaced with uncooked. Defaults to no objects */
+	UE_INTERNAL
+	COREUOBJECT_API virtual void GetAdditionalObjectsToUncook(TArray<UObject*>& AdditionalObjects) const { }
+
 #else
 	FORCEINLINE bool Modify(bool bAlwaysMarkDirty = true) { return false; }
 	FORCEINLINE bool IsCapturingAsRootObjectForTransaction() const { return false; };
