- Removed unused API i had added and was a bit confusing

diff --git a/Engine/Plugins/Editor/ContentBrowser/ContentBrowserAssetDataSource/Source/ContentBrowserAssetDataSource/Private/AssetFileContextMenu.cpp b/Engine/Plugins/Editor/ContentBrowser/ContentBrowserAssetDataSource/Source/ContentBrowserAssetDataSource/Private/AssetFileContextMenu.cpp
--- a/Engine/Plugins/Editor/ContentBrowser/ContentBrowserAssetDataSource/Source/ContentBrowserAssetDataSource/Private/AssetFileContextMenu.cpp
+++ b/Engine/Plugins/Editor/ContentBrowser/ContentBrowserAssetDataSource/Source/ContentBrowserAssetDataSource/Private/AssetFileContextMenu.cpp
@@ -1664,15 +1664,8 @@ static void ProcessObjectForUncook(const UObject* Obj, TArray<UPackage*>& Packag
 		PackagesToReload.AddUnique(Pkg);
 	}
 
-	// now unload any additional packages for anything this object needs
-	TArray<UObject*> AdditionalObjects;
-	PRAGMA_DISABLE_INTERNAL_WARNINGS
-	Obj->GetAdditionalObjectsToUncook(AdditionalObjects);
-	PRAGMA_ENABLE_INTERNAL_WARNINGS
-	for (UObject* RelatedObj : AdditionalObjects)
-	{
-		ProcessObjectForUncook(RelatedObj, PackagesToReload);
-	}
+	// if the object needs anymore objects to also be unloaded, call this function recursively
+	// here, with those objects
 }
 
 static void ReloadPackages(TArray<FAssetData>& Assets, bool bUncookPackages)
diff --git a/Engine/Source/Runtime/CoreUObject/Public/UObject/Object.h b/Engine/Source/Runtime/CoreUObject/Public/UObject/Object.h
--- a/Engine/Source/Runtime/CoreUObject/Public/UObject/Object.h
+++ b/Engine/Source/Runtime/CoreUObject/Public/UObject/Object.h
@@ -316,10 +316,6 @@ public:
 	/** Utility to allow overrides of Modify to avoid doing work if this object cannot be safely modified */
 	COREUOBJECT_API bool CanModify() const;
 
-	/** Returns a list of additional objects that will need to be replaced with uncooked versions when this object is being replaced with uncooked. Defaults to no objects */
-	UE_INTERNAL
-	COREUOBJECT_API virtual void GetAdditionalObjectsToUncook(TArray<UObject*>& AdditionalObjects) const { }
-
 #else
 	FORCEINLINE bool Modify(bool bAlwaysMarkDirty = true) { return false; }
 	FORCEINLINE bool IsCapturingAsRootObjectForTransaction() const { return false; };
