[ContentBrowser]
- Parallelize OnAssetsAdded
- Saved 6s on boot on some internal project

diff --git a/Engine/Plugins/Editor/ContentBrowser/ContentBrowserAssetDataSource/Source/ContentBrowserAssetDataSource/Private/ContentBrowserAssetDataSource.cpp b/Engine/Plugins/Editor/ContentBrowser/ContentBrowserAssetDataSource/Source/ContentBrowserAssetDataSource/Private/ContentBrowserAssetDataSource.cpp
--- a/Engine/Plugins/Editor/ContentBrowser/ContentBrowserAssetDataSource/Source/ContentBrowserAssetDataSource/Private/ContentBrowserAssetDataSource.cpp
+++ b/Engine/Plugins/Editor/ContentBrowser/ContentBrowserAssetDataSource/Source/ContentBrowserAssetDataSource/Private/ContentBrowserAssetDataSource.cpp
@@ -17,6 +17,7 @@
 #include "ContentBrowserItemPath.h"
 #include "UObject/GCObjectScopeGuard.h"
 #include "UObject/ObjectSaveContext.h"
+#include "Containers/ConsumeAllMpmcQueue.h"
 #include "Factories/Factory.h"
 #include "HAL/FileManager.h"
 #include "Editor.h"
@@ -3339,23 +3340,41 @@ void UContentBrowserAssetDataSource::OnAssetsAdded(TConstArrayView<FAssetData> I
 	LLM_SCOPE_BYNAME(TEXT("UContentBrowserAssetDataSource"))
 
 	FAssetPropertyTagCache& Cache = FAssetPropertyTagCache::Get();
-	for (const FAssetData& InAssetData : InAssets)
+
+	struct FUpdate
 	{
-		UE_LOG(LogContentBrowserAssetDataSource, VeryVerbose, TEXT("OnAssetsAdded: %s"), *WriteToString<256>(InAssetData.GetSoftObjectPath()));
+		int32 Index;
+		FContentBrowserItemDataUpdate ItemDataUpdate;
+	};
+	UE::TConsumeAllMpmcQueue<FUpdate> Updates;
 
-		if (InAssetData.GetOptionalOuterPathName().IsNone())
+	ParallelFor(InAssets.Num(),
+		[this, &Cache, &InAssets, &Updates](int32 Index)
 		{
-			Cache.TryCacheClass(InAssetData.AssetClassPath);
+			const FAssetData& InAssetData = InAssets[Index];
+			UE_LOG(LogContentBrowserAssetDataSource, VeryVerbose, TEXT("OnAssetsAdded: %s"), *WriteToString<256>(InAssetData.GetSoftObjectPath()));
+
+			if (InAssetData.GetOptionalOuterPathName().IsNone())
+			{
+				Cache.TryCacheClass(InAssetData.AssetClassPath);
+			}
+
+			if (ContentBrowserAssetData::IsPrimaryAsset(InAssetData))
+			{
+				Updates.ProduceItem(Index, FContentBrowserItemDataUpdate::MakeItemAddedUpdate(CreateAssetFileItem(InAssetData)));
+			}
 		}
+	);
 
-		if (ContentBrowserAssetData::IsPrimaryAsset(InAssetData))
+	Updates.ConsumeAllLifo(
+		[this, &InAssets](FUpdate&& Update)
 		{
 			// The owner folder of this asset is no longer considered empty
-			OnPathPopulated(InAssetData);
+			OnPathPopulated(InAssets[Update.Index]);
 
-			QueueItemDataUpdate(FContentBrowserItemDataUpdate::MakeItemAddedUpdate(CreateAssetFileItem(InAssetData)));
+			QueueItemDataUpdate(MoveTemp(Update.ItemDataUpdate));
 		}
-	}
+	);
 }
 
 void UContentBrowserAssetDataSource::OnAssetRemoved(const FAssetData& InAssetData)
