Fixed a tooltip localization key

diff --git a/Engine/Plugins/Editor/ContentBrowser/ContentBrowserAssetDataSource/Source/ContentBrowserAssetDataSource/Private/AssetFileContextMenu.cpp b/Engine/Plugins/Editor/ContentBrowser/ContentBrowserAssetDataSource/Source/ContentBrowserAssetDataSource/Private/AssetFileContextMenu.cpp
--- a/Engine/Plugins/Editor/ContentBrowser/ContentBrowserAssetDataSource/Source/ContentBrowserAssetDataSource/Private/AssetFileContextMenu.cpp
+++ b/Engine/Plugins/Editor/ContentBrowser/ContentBrowserAssetDataSource/Source/ContentBrowserAssetDataSource/Private/AssetFileContextMenu.cpp
@@ -446,7 +446,7 @@ void FAssetFileContextMenu::MakeAssetActionsSubMenu(UToolMenu* Menu)
 			Section.AddMenuEntry(
 				"ReloadUncooked",
 				LOCTEXT("ReloadUncooked", "Reload Uncooked Version"),
-				LOCTEXT("ReloadTooltip", "Reload the selected assets from their uncooked file on disk."),
+				LOCTEXT("ReloadUncookedTooltip", "Reload the selected assets from their uncooked file on disk."),
 				FSlateIcon(),
 				FUIAction(
 					FExecuteAction::CreateSP(this, &FAssetFileContextMenu::ExecuteReloadUncooked),
