Add a Copy Asset Registry Tags context menu action for CB assets

diff --git a/Engine/Plugins/Editor/ContentBrowser/ContentBrowserAssetDataSource/Source/ContentBrowserAssetDataSource/ContentBrowserAssetDataSource.Build.cs b/Engine/Plugins/Editor/ContentBrowser/ContentBrowserAssetDataSource/Source/ContentBrowserAssetDataSource/ContentBrowserAssetDataSource.Build.cs
--- a/Engine/Plugins/Editor/ContentBrowser/ContentBrowserAssetDataSource/Source/ContentBrowserAssetDataSource/ContentBrowserAssetDataSource.Build.cs
+++ b/Engine/Plugins/Editor/ContentBrowser/ContentBrowserAssetDataSource/Source/ContentBrowserAssetDataSource/ContentBrowserAssetDataSource.Build.cs
@@ -42,6 +42,7 @@ namespace UnrealBuildTool.Rules
 					"Slate",
 					"SlateCore",
 					"InputCore",
+					"Json",
 					"SourceControl",
 					"SourceControlWindows",
 				}
diff --git a/Engine/Plugins/Editor/ContentBrowser/ContentBrowserAssetDataSource/Source/ContentBrowserAssetDataSource/Private/AssetFileContextMenu.cpp b/Engine/Plugins/Editor/ContentBrowser/ContentBrowserAssetDataSource/Source/ContentBrowserAssetDataSource/Private/AssetFileContextMenu.cpp
--- a/Engine/Plugins/Editor/ContentBrowser/ContentBrowserAssetDataSource/Source/ContentBrowserAssetDataSource/Private/AssetFileContextMenu.cpp
+++ b/Engine/Plugins/Editor/ContentBrowser/ContentBrowserAssetDataSource/Source/ContentBrowserAssetDataSource/Private/AssetFileContextMenu.cpp
@@ -41,6 +41,7 @@
 #include "ReferencedAssetsUtils.h"
 #include "Internationalization/PackageLocalizationUtil.h"
 #include "Internationalization/TextLocalizationResource.h"
+#include "Serialization/JsonSerializer.h"
 #include "Serialization/PropertyLocalizationDataGathering.h"
 
 #include "Kismet2/KismetEditorUtilities.h"
@@ -527,6 +528,17 @@ void FAssetFileContextMenu::MakeAssetActionsSubMenu(UToolMenu* Menu)
 			);
 		}
 
+		if (bCanBeModified)
+		{
+			Section.AddMenuEntry(
+				"CopyAssetRegistryTags",
+				LOCTEXT("CopyAssetRegistryTags", "Copy Asset Registry Tags"),
+				LOCTEXT("CopyAssetRegistryTagsTooltip", "Copies to clipboard all asset registry tags that belong to this asset."),
+				FSlateIcon(),
+				FUIAction(FExecuteAction::CreateSP(this, &FAssetFileContextMenu::ExecuteCopyAssetRegistryTags))
+			);
+		}
+
 		// Chunk actions
 		if (bCanBeModified && GetDefault<UEditorExperimentalSettings>()->bContextMenuChunkAssignments)
 		{
@@ -1595,6 +1607,40 @@ void FAssetFileContextMenu::ExecuteShowAssetMetaData()
 	}
 }
 
+void FAssetFileContextMenu::ExecuteCopyAssetRegistryTags()
+{
+	const TSharedRef<FJsonObject> JsonObject = MakeShared<FJsonObject>();
+	for (const FAssetData& AssetData : SelectedAssets)
+	{
+		TSharedPtr<FJsonObject> AssetObject;
+
+		if (SelectedAssets.Num() != 1)
+		{
+			AssetObject = MakeShared<FJsonObject>();
+			JsonObject->SetObjectField(AssetData.PackageName.ToString(), AssetObject);
+		}
+		else
+		{
+			AssetObject = JsonObject;
+		}
+
+		AssetData.EnumerateTags([&AssetObject](const TPair<FName, FAssetTagValueRef>& TagPair)
+		{
+			AssetObject->SetStringField(TagPair.Key.ToString(), TagPair.Value.AsString());
+		});
+	}
+
+	FString JsonString;
+	const TSharedRef<TJsonWriter<>> JsonWriter = TJsonWriterFactory<>::Create(&JsonString);
+	FJsonSerializer::Serialize(JsonObject, JsonWriter);
+
+	FPlatformApplicationMisc::ClipboardCopy(*JsonString);
+
+	FNotificationInfo Info(LOCTEXT("CopyAssetRegistryTagsNotification", "Asset Registry Tags have been copied to clipboard."));
+	Info.ExpireDuration = 3.0f;
+	FSlateNotificationManager::Get().AddNotification(Info);
+}
+
 bool FAssetFileContextMenu::CanModifyPath(const FString& InPath) const
 {
 	FAssetToolsModule& AssetToolsModule = FModuleManager::LoadModuleChecked<FAssetToolsModule>("AssetTools");
diff --git a/Engine/Plugins/Editor/ContentBrowser/ContentBrowserAssetDataSource/Source/ContentBrowserAssetDataSource/Public/AssetFileContextMenu.h b/Engine/Plugins/Editor/ContentBrowser/ContentBrowserAssetDataSource/Source/ContentBrowserAssetDataSource/Public/AssetFileContextMenu.h
--- a/Engine/Plugins/Editor/ContentBrowser/ContentBrowserAssetDataSource/Source/ContentBrowserAssetDataSource/Public/AssetFileContextMenu.h
+++ b/Engine/Plugins/Editor/ContentBrowser/ContentBrowserAssetDataSource/Source/ContentBrowserAssetDataSource/Public/AssetFileContextMenu.h
@@ -141,6 +141,9 @@ private:
 	/** Handler for when "Show MetaData" is selected */
 	UE_API void ExecuteShowAssetMetaData();
 
+	/** Handler for when "Copy Asset Registry Tags" is selected */
+	void ExecuteCopyAssetRegistryTags();
+
 	/** Handler for when "Diff Selected" is selected */
 	UE_API void ExecuteDiffSelected() const;
 
