Extend the Create Asset menu to allow non-asset menu entries for wizards that create potentially multiple assets of various types

diff --git a/Engine/Plugins/Editor/ContentBrowser/ContentBrowserAssetDataSource/Source/ContentBrowserAssetDataSource/Private/NewAssetContextMenu.cpp b/Engine/Plugins/Editor/ContentBrowser/ContentBrowserAssetDataSource/Source/ContentBrowserAssetDataSource/Private/NewAssetContextMenu.cpp
--- a/Engine/Plugins/Editor/ContentBrowser/ContentBrowserAssetDataSource/Source/ContentBrowserAssetDataSource/Private/NewAssetContextMenu.cpp
+++ b/Engine/Plugins/Editor/ContentBrowser/ContentBrowserAssetDataSource/Source/ContentBrowserAssetDataSource/Private/NewAssetContextMenu.cpp
@@ -13,6 +13,7 @@
 #include "Factories/Factory.h"
 #include "IAssetTools.h"
 #include "AssetToolsModule.h"
+#include "ContentBrowserModule.h"
 #include "Widgets/SToolTip.h"
 #include "IDocumentation.h"
 #include "ContentBrowserUtils.h"
@@ -33,10 +34,19 @@ struct FFactoryItem
 	}
 };
 
+struct FWizardItem
+{
+	FText DisplayName;
+	FText Description;
+	FSlateIcon Icon;
+	FSimpleDelegate OnClicked;
+};
+
 struct FCategorySubMenuItem
 {
 	FText Name;
 	TArray<FFactoryItem> Factories;
+	TArray<FWizardItem> Wizards;
 	TMap<FString, TSharedPtr<FCategorySubMenuItem>> Children;
 
 	void SortSubMenus(FCategorySubMenuItem* SubMenu = nullptr)
@@ -321,7 +331,8 @@ void FNewAssetContextMenu::MakeContextMenu(
 
 			const IAssetTools& AssetTools = AssetToolsModule.Get();
 			const TArray<UFactory*> NewAssetFactories = AssetTools.GetNewAssetFactories();
-
+			
+			
 			for (const FAdvancedAssetCategory& AdvancedAssetCategory : AdvancedAssetCategories)
 			{
 				const bool FindFirstOnly = true;
@@ -390,6 +401,36 @@ void FNewAssetContextMenu::CreateNewAssetMenuCategory(UToolMenu* Menu, FName Sec
 	}
 	ParentMenuData->SortSubMenus();
 
+	// Find wizards
+	FContentBrowserModule& ContentBrowserModule = FModuleManager::GetModuleChecked<FContentBrowserModule>(TEXT("ContentBrowser"));
+	for (const FContentBrowserModule::FWizard& Wizard : ContentBrowserModule.GetWizards())
+	{
+		if (Wizard.CategoryPath.CategoryType == AssetTypeCategory)
+		{
+			FCategorySubMenuItem* SubMenu = ParentMenuData.Get();
+
+			const FString SourceString = Wizard.CategoryPath.CategoryName.BuildSourceString();
+			if (TSharedPtr<FCategorySubMenuItem> SubMenuData = SubMenu->Children.FindRef(SourceString))
+			{
+				check(SubMenuData.IsValid());
+				SubMenu = SubMenuData.Get();
+			}
+			else
+			{
+				TSharedPtr<FCategorySubMenuItem> NewSubMenu = MakeShared<FCategorySubMenuItem>();
+				NewSubMenu->Name = Wizard.CategoryPath.CategoryName;
+				SubMenu->Children.Add(SourceString, NewSubMenu);
+				SubMenu = NewSubMenu.Get();
+			}
+			
+			FWizardItem& WizardItem = SubMenu->Wizards.AddDefaulted_GetRef();
+			WizardItem.DisplayName = Wizard.DisplayName;
+			WizardItem.Description = Wizard.Description;
+			WizardItem.Icon = Wizard.Icon;
+			WizardItem.OnClicked = Wizard.OnOpen;
+		}
+	}
+
 	CreateNewAssetMenus(Menu, SectionName, ParentMenuData, InPath, InOnNewAssetRequested, InCanExecuteAction);
 }
 
@@ -398,6 +439,25 @@ void FNewAssetContextMenu::CreateNewAssetMenus(UToolMenu* Menu, FName SectionNam
 	QUICK_SCOPE_CYCLE_COUNTER(CreateNewAssetMenus);
 
 	FToolMenuSection& Section = Menu->FindOrAddSection(SectionName);
+
+	for (const FWizardItem& WizardItem : SubMenuData->Wizards)
+	{
+		Section.AddMenuEntry(
+			NAME_None,
+			WizardItem.DisplayName,
+			WizardItem.Description,
+			WizardItem.Icon,
+			FUIAction(FExecuteAction::CreateLambda([WizardItem]()
+			{
+				WizardItem.OnClicked.ExecuteIfBound();
+			})));
+	}
+
+	if (!SubMenuData->Wizards.IsEmpty())
+	{
+		Section.AddSeparator(NAME_None);
+	}
+	
 	for (const FFactoryItem& FactoryItem : SubMenuData->Factories)
 	{
 		TWeakObjectPtr<UClass> WeakFactoryClass = FactoryItem.Factory.GetClass();
diff --git a/Engine/Source/Editor/ContentBrowser/Private/ContentBrowserModule.cpp b/Engine/Source/Editor/ContentBrowser/Private/ContentBrowserModule.cpp
--- a/Engine/Source/Editor/ContentBrowser/Private/ContentBrowserModule.cpp
+++ b/Engine/Source/Editor/ContentBrowser/Private/ContentBrowserModule.cpp
@@ -85,3 +85,21 @@ void FContentBrowserModule::ContentBrowserSettingChanged(FName InName)
 
 	OnContentBrowserSettingChanged.Broadcast(InName);
 }
+
+void FContentBrowserModule::RegisterWizard(const FWizard& InWizard)
+{
+	Wizards.Add(InWizard);
+}
+
+void FContentBrowserModule::UnregisterWizard(const FName WizardIdentifier)
+{
+	Wizards.RemoveAll([WizardIdentifier](const FWizard& Wizard)
+	{
+		return Wizard.WizardIdentifier == WizardIdentifier;
+	});
+}
+
+const TArray<FContentBrowserModule::FWizard>& FContentBrowserModule::GetWizards() const
+{
+	return Wizards;
+}
diff --git a/Engine/Source/Editor/ContentBrowser/Public/ContentBrowserModule.h b/Engine/Source/Editor/ContentBrowser/Public/ContentBrowserModule.h
--- a/Engine/Source/Editor/ContentBrowser/Public/ContentBrowserModule.h
+++ b/Engine/Source/Editor/ContentBrowser/Public/ContentBrowserModule.h
@@ -7,6 +7,8 @@
 #include "Containers/Array.h"
 #include "ContentBrowserDelegates.h"
 #include "CoreMinimal.h"
+#include "AssetTypeCategories.h"
+#include "IAssetTools.h"
 #include "Delegates/Delegate.h"
 #include "Experimental/ContentBrowserViewExtender.h"
 #include "HAL/PlatformCrt.h"
@@ -222,9 +224,24 @@ public:
 	{
 		return OnSetFolderColor;
 	}
-	
+
+	struct FWizard
+	{
+		FName WizardIdentifier;
+		FText DisplayName;
+		FText Description;
+		FSlateIcon Icon;
+		FAdvancedAssetCategory CategoryPath;
+		FSimpleDelegate OnOpen;
+	};
+
+	CONTENTBROWSER_API void RegisterWizard(const FWizard& InWizard);
+	CONTENTBROWSER_API void UnregisterWizard(const FName WizardIdentifier);
+	CONTENTBROWSER_API const TArray<FWizard>& GetWizards() const;
 
 private:
+	TArray<FWizard> Wizards;
+	
 	/** Handle changes to content browser settings */
 	void ContentBrowserSettingChanged(FName InName);
 	
