Updated UContentBrowserAliasDataSource and UContentBrowserAssetDataSource to hold ContentBrowserModule pointers instead of using GetModuleChecked for on demand usage.

diff --git a/Engine/Plugins/Editor/ContentBrowser/ContentBrowserAliasDataSource/Source/ContentBrowserAliasDataSource/Private/ContentBrowserAliasDataSource.cpp b/Engine/Plugins/Editor/ContentBrowser/ContentBrowserAliasDataSource/Source/ContentBrowserAliasDataSource/Private/ContentBrowserAliasDataSource.cpp
--- a/Engine/Plugins/Editor/ContentBrowser/ContentBrowserAliasDataSource/Source/ContentBrowserAliasDataSource/Private/ContentBrowserAliasDataSource.cpp
+++ b/Engine/Plugins/Editor/ContentBrowser/ContentBrowserAliasDataSource/Source/ContentBrowserAliasDataSource/Private/ContentBrowserAliasDataSource.cpp
@@ -7,6 +7,7 @@
 #include "ContentBrowserAssetDataCore.h"
 #include "ContentBrowserDataSubsystem.h"
 #include "ContentBrowserItemPath.h"
+#include "ContentBrowserModule.h"
 
 #include "AssetToolsModule.h"
 #include "CollectionManagerModule.h"
@@ -73,6 +74,7 @@ void UContentBrowserAliasDataSource::Initialize(const bool InAutoRegister)
 	{
 		AssetRegistry = &FModuleManager::LoadModuleChecked<FAssetRegistryModule>(AssetRegistryConstants::ModuleName).Get();
 		AssetTools = &FModuleManager::GetModuleChecked<FAssetToolsModule>("AssetTools").Get();
+		ContentBrowserModule = &FModuleManager::Get().GetModuleChecked<FContentBrowserModule>(TEXT("ContentBrowser"));
 
 		AssetRegistry->OnAssetAdded().AddUObject(this, &UContentBrowserAliasDataSource::OnAssetAdded);
 		AssetRegistry->OnAssetRemoved().AddUObject(this, &UContentBrowserAliasDataSource::OnAssetRemoved);
@@ -897,7 +899,7 @@ void UContentBrowserAliasDataSource::OnObjectPropertyChanged(UObject* InObject,
 
 bool UContentBrowserAliasDataSource::GetItemAttribute(const FContentBrowserItemData& InItem, const bool InIncludeMetaData, const FName InAttributeKey, FContentBrowserItemDataAttributeValue& OutAttributeValue)
 {
-	return ContentBrowserAssetData::GetItemAttribute(this, InItem, InIncludeMetaData, InAttributeKey, OutAttributeValue);
+	return ContentBrowserAssetData::GetItemAttribute(ContentBrowserModule, this, InItem, InIncludeMetaData, InAttributeKey, OutAttributeValue);
 }
 
 bool UContentBrowserAliasDataSource::GetItemAttributes(const FContentBrowserItemData& InItem, const bool InIncludeMetaData, FContentBrowserItemDataAttributeValues& OutAttributeValues)
diff --git a/Engine/Plugins/Editor/ContentBrowser/ContentBrowserAliasDataSource/Source/ContentBrowserAliasDataSource/Public/ContentBrowserAliasDataSource.h b/Engine/Plugins/Editor/ContentBrowser/ContentBrowserAliasDataSource/Source/ContentBrowserAliasDataSource/Public/ContentBrowserAliasDataSource.h
--- a/Engine/Plugins/Editor/ContentBrowser/ContentBrowserAliasDataSource/Source/ContentBrowserAliasDataSource/Public/ContentBrowserAliasDataSource.h
+++ b/Engine/Plugins/Editor/ContentBrowser/ContentBrowserAliasDataSource/Source/ContentBrowserAliasDataSource/Public/ContentBrowserAliasDataSource.h
@@ -21,6 +21,7 @@ DECLARE_LOG_CATEGORY_EXTERN(LogContentBrowserAliasDataSource, Log, All);
 
 class IAssetTools;
 struct FContentBrowserCompiledAssetDataFilter;
+class FContentBrowserModule;
 
 /** A unique alias is a pair of SourceObjectPath:AliasPath, eg /Game/MyAsset.MyAsset:/Game/SomeFolder/MyAlias */
 typedef TPair<FSoftObjectPath, FName> FContentBrowserUniqueAlias;
@@ -197,6 +198,7 @@ private:
 
 	IAssetRegistry* AssetRegistry = nullptr;
 	IAssetTools* AssetTools = nullptr;
+	FContentBrowserModule* ContentBrowserModule = nullptr;
 
 	struct FAliasData
 	{
diff --git a/Engine/Plugins/Editor/ContentBrowser/ContentBrowserAssetDataSource/Source/ContentBrowserAssetDataSource/Private/ContentBrowserAssetDataCore.cpp b/Engine/Plugins/Editor/ContentBrowser/ContentBrowserAssetDataSource/Source/ContentBrowserAssetDataSource/Private/ContentBrowserAssetDataCore.cpp
--- a/Engine/Plugins/Editor/ContentBrowser/ContentBrowserAssetDataSource/Source/ContentBrowserAssetDataSource/Private/ContentBrowserAssetDataCore.cpp
+++ b/Engine/Plugins/Editor/ContentBrowser/ContentBrowserAssetDataSource/Source/ContentBrowserAssetDataSource/Private/ContentBrowserAssetDataCore.cpp
@@ -845,7 +845,7 @@ bool CanPrivatizeAssetFileItem(IAssetTools* InAssetTools, const FContentBrowserA
 	return true;
 }
 
-bool DeleteItems(IAssetTools* InAssetTools, IAssetRegistry* InAssetRegistry, const UContentBrowserDataSource* InOwnerDataSource, TArrayView<const FContentBrowserItemData> InItems)
+bool DeleteItems(FContentBrowserModule* InContentBrowserModule, IAssetTools* InAssetTools, IAssetRegistry* InAssetRegistry, const UContentBrowserDataSource* InOwnerDataSource, TArrayView<const FContentBrowserItemData> InItems)
 {
 	TArray<TSharedRef<const FContentBrowserAssetFolderItemDataPayload>, TInlineAllocator<16>> FolderPayloads;
 	TArray<FAssetData> AssetDatas;
@@ -887,7 +887,7 @@ bool DeleteItems(IAssetTools* InAssetTools, IAssetRegistry* InAssetRegistry, con
 
 	if (FolderPayloads.Num() > 0)
 	{
-		bDidDelete |= DeleteAssetFolderItems(FolderPayloads);
+		bDidDelete |= DeleteAssetFolderItems(InContentBrowserModule, FolderPayloads);
 	}
 
 	if (AssetDatas.Num() > 0)
@@ -933,7 +933,7 @@ bool PrivatizeItems(IAssetTools* InAssetTools, IAssetRegistry* InAssetRegistry,
 	return ObjectTools::PrivatizeAssets(AssetsToPrivatize, /*bShowConfirmation*/true, InAssetAccessSpecifier) > 0;
 }
 
-bool DeleteAssetFolderItems(TArrayView<const TSharedRef<const FContentBrowserAssetFolderItemDataPayload>> InFolderPayloads)
+bool DeleteAssetFolderItems(FContentBrowserModule* InContentBrowserModule, TArrayView<const TSharedRef<const FContentBrowserAssetFolderItemDataPayload>> InFolderPayloads)
 {
 	TArray<FContentBrowserItemPath> ItemPathsToDelete;
 	ItemPathsToDelete.Reserve(InFolderPayloads.Num());
@@ -943,8 +943,7 @@ bool DeleteAssetFolderItems(TArrayView<const TSharedRef<const FContentBrowserAss
 		ItemPathsToDelete.Add(FContentBrowserItemPath(FolderPayload->GetInternalPath(), EContentBrowserPathType::Internal));
 	}
 
-	FContentBrowserModule& ContentBrowserModule = FModuleManager::Get().GetModuleChecked<FContentBrowserModule>(TEXT("ContentBrowser"));
-	ContentBrowserModule.GetOnDeleteFolders().Broadcast(ItemPathsToDelete);
+	InContentBrowserModule->GetOnDeleteFolders().Broadcast(ItemPathsToDelete);
 
 	TArray<FString> FoldersToDelete;
 	Algo::Transform(ItemPathsToDelete, FoldersToDelete, [](FContentBrowserItemPath ItemPath) { return ItemPath.GetInternalPathString(); });
@@ -1734,11 +1733,10 @@ void GetGenericItemAttribute(const FName InTagKey, const FString& InTagValue, co
 	}
 }
 
-bool GetItemAttribute(const UContentBrowserDataSource* InOwnerDataSource, const FContentBrowserItemData& InItem, const bool InIncludeMetaData, const FName InAttributeKey, FContentBrowserItemDataAttributeValue& OutAttributeValue)
+bool GetItemAttribute(FContentBrowserModule* InContentBrowserModule, const UContentBrowserDataSource* InOwnerDataSource, const FContentBrowserItemData& InItem, const bool InIncludeMetaData, const FName InAttributeKey, FContentBrowserItemDataAttributeValue& OutAttributeValue)
 {
 	// First see if another module hooked this to provide its own customizations. If it did, we can be done here.
-	FContentBrowserModule& ContentBrowserModule = FModuleManager::Get().GetModuleChecked<FContentBrowserModule>(TEXT("ContentBrowser"));
-	for(const FGetItemDataAttribute& Delegate : ContentBrowserModule.GetAllItemDataAttributeDelegates())
+	for(const FGetItemDataAttribute& Delegate : InContentBrowserModule->GetAllItemDataAttributeDelegates())
 	{
 		if (Delegate.Execute(InItem, InIncludeMetaData, InAttributeKey, OutAttributeValue))
 		{
diff --git a/Engine/Plugins/Editor/ContentBrowser/ContentBrowserAssetDataSource/Source/ContentBrowserAssetDataSource/Private/ContentBrowserAssetDataSource.cpp b/Engine/Plugins/Editor/ContentBrowser/ContentBrowserAssetDataSource/Source/ContentBrowserAssetDataSource/Private/ContentBrowserAssetDataSource.cpp
--- a/Engine/Plugins/Editor/ContentBrowser/ContentBrowserAssetDataSource/Source/ContentBrowserAssetDataSource/Private/ContentBrowserAssetDataSource.cpp
+++ b/Engine/Plugins/Editor/ContentBrowser/ContentBrowserAssetDataSource/Source/ContentBrowserAssetDataSource/Private/ContentBrowserAssetDataSource.cpp
@@ -164,6 +164,11 @@ void UContentBrowserAssetDataSource::Initialize(const bool InAutoRegister)
 		AssetTools = &FModuleManager::GetModuleChecked<FAssetToolsModule>(NAME_AssetTools).Get();
 	}
 
+	{
+		static const FName NAME_ContentBrowser = "ContentBrowser";
+		ContentBrowserModule = &FModuleManager::Get().GetModuleChecked<FContentBrowserModule>(NAME_ContentBrowser);
+	}
+
 	CollectionManager = &FCollectionManagerModule::GetModule().Get();
 
 	// Listen for asset registry updates
@@ -2709,7 +2714,7 @@ bool UContentBrowserAssetDataSource::ConvertItemForFilter(FContentBrowserItemDat
 
 bool UContentBrowserAssetDataSource::GetItemAttribute(const FContentBrowserItemData& InItem, const bool InIncludeMetaData, const FName InAttributeKey, FContentBrowserItemDataAttributeValue& OutAttributeValue)
 {
-	return ContentBrowserAssetData::GetItemAttribute(this, InItem, InIncludeMetaData, InAttributeKey, OutAttributeValue);
+	return ContentBrowserAssetData::GetItemAttribute(ContentBrowserModule, this, InItem, InIncludeMetaData, InAttributeKey, OutAttributeValue);
 }
 
 bool UContentBrowserAssetDataSource::GetItemAttributes(const FContentBrowserItemData& InItem, const bool InIncludeMetaData, FContentBrowserItemDataAttributeValues& OutAttributeValues)
@@ -2860,12 +2865,12 @@ bool UContentBrowserAssetDataSource::CanDeleteItem(const FContentBrowserItemData
 
 bool UContentBrowserAssetDataSource::DeleteItem(const FContentBrowserItemData& InItem)
 {
-	return ContentBrowserAssetData::DeleteItems(AssetTools, AssetRegistry, this, MakeArrayView(&InItem, 1));
+	return ContentBrowserAssetData::DeleteItems(ContentBrowserModule, AssetTools, AssetRegistry, this, MakeArrayView(&InItem, 1));
 }
 
 bool UContentBrowserAssetDataSource::BulkDeleteItems(TArrayView<const FContentBrowserItemData> InItems)
 {
-	return ContentBrowserAssetData::DeleteItems(AssetTools, AssetRegistry, this, InItems);
+	return ContentBrowserAssetData::DeleteItems(ContentBrowserModule, AssetTools, AssetRegistry, this, InItems);
 }
 
 bool UContentBrowserAssetDataSource::CanPrivatizeItem(const FContentBrowserItemData& InItem, FText* OutErrorMsg)
@@ -4013,8 +4018,7 @@ void UContentBrowserAssetDataSource::OnBeginCreateAsset(const FName InDefaultAss
 
 	FAssetPropertyTagCache::Get().TryCacheClass(FTopLevelAssetPath(ClassToUse));
 
-	FContentBrowserModule& ContentBrowserModule = FModuleManager::Get().GetModuleChecked<FContentBrowserModule>(TEXT("ContentBrowser"));
-	FContentBrowserItemPath AssetPathToUse = ContentBrowserModule.Get().GetInitialPathToSaveAsset(FContentBrowserItemPath(InPackagePath, EContentBrowserPathType::Internal));
+	FContentBrowserItemPath AssetPathToUse = ContentBrowserModule->Get().GetInitialPathToSaveAsset(FContentBrowserItemPath(InPackagePath, EContentBrowserPathType::Internal));
 
 	const bool bUseAsyncCreate = InFactory && EnumHasAnyFlags(InFactory->GetSupportedWorkflows(), EFactoryCreateWorkflow::Asynchronous);
 	const bool bShowDialogToPickPath = !AssetPathToUse.HasInternalPath() || (AssetPathToUse.GetInternalPathName() != InPackagePath);
diff --git a/Engine/Plugins/Editor/ContentBrowser/ContentBrowserAssetDataSource/Source/ContentBrowserAssetDataSource/Public/ContentBrowserAssetDataCore.h b/Engine/Plugins/Editor/ContentBrowser/ContentBrowserAssetDataSource/Source/ContentBrowserAssetDataSource/Public/ContentBrowserAssetDataCore.h
--- a/Engine/Plugins/Editor/ContentBrowser/ContentBrowserAssetDataSource/Source/ContentBrowserAssetDataSource/Public/ContentBrowserAssetDataCore.h
+++ b/Engine/Plugins/Editor/ContentBrowser/ContentBrowserAssetDataSource/Source/ContentBrowserAssetDataSource/Public/ContentBrowserAssetDataCore.h
@@ -14,6 +14,7 @@ class IAssetRegistry;
 class FAssetThumbnail;
 class FAssetFolderContextMenu;
 class FAssetFileContextMenu;
+class FContentBrowserModule;
 class UObject;
 class UToolMenu;
 class UContentBrowserDataSource;
@@ -123,9 +124,9 @@ CONTENTBROWSERASSETDATASOURCE_API FContentBrowserItemData CreateUnsupportedAsset
 
 	CONTENTBROWSERASSETDATASOURCE_API bool CanDeleteUnsupportedAssetFileItem(IAssetTools* InAssetTools, const FContentBrowserUnsupportedAssetFileItemDataPayload& InUnsupportedAssetPayload, FText* OutErrorMsg);
 
-	CONTENTBROWSERASSETDATASOURCE_API bool DeleteItems(IAssetTools* InAssetTools, IAssetRegistry* InAssetRegistry, const UContentBrowserDataSource* InOwnerDataSource, TArrayView<const FContentBrowserItemData> InItems);
+	CONTENTBROWSERASSETDATASOURCE_API bool DeleteItems(FContentBrowserModule* InContentBrowserModule, IAssetTools* InAssetTools, IAssetRegistry* InAssetRegistry, const UContentBrowserDataSource* InOwnerDataSource, TArrayView<const FContentBrowserItemData> InItems);
 
-	CONTENTBROWSERASSETDATASOURCE_API bool DeleteAssetFolderItems(TArrayView<const TSharedRef<const FContentBrowserAssetFolderItemDataPayload>> InFolderPayloads);
+	CONTENTBROWSERASSETDATASOURCE_API bool DeleteAssetFolderItems(FContentBrowserModule* InContentBrowserModule, TArrayView<const TSharedRef<const FContentBrowserAssetFolderItemDataPayload>> InFolderPayloads);
 
 	CONTENTBROWSERASSETDATASOURCE_API bool DeleteAssetFileItems(TArrayView<const TSharedRef<const FContentBrowserAssetFileItemDataPayload>> InAssetPayloads);
 
@@ -200,7 +201,7 @@ CONTENTBROWSERASSETDATASOURCE_API FContentBrowserItemData CreateUnsupportedAsset
 
 	CONTENTBROWSERASSETDATASOURCE_API bool GetUnsupportedAssetFileItemPhysicalPath(const FContentBrowserUnsupportedAssetFileItemDataPayload& InUnsupportedAssetPayload, FString& OutDiskPath);
 
-	CONTENTBROWSERASSETDATASOURCE_API bool GetItemAttribute(const UContentBrowserDataSource* InOwnerDataSource, const FContentBrowserItemData& InItem, const bool InIncludeMetaData, const FName InAttributeKey, FContentBrowserItemDataAttributeValue& OutAttributeValue);
+	CONTENTBROWSERASSETDATASOURCE_API bool GetItemAttribute(FContentBrowserModule* InContentBrowserModule, const UContentBrowserDataSource* InOwnerDataSource, const FContentBrowserItemData& InItem, const bool InIncludeMetaData, const FName InAttributeKey, FContentBrowserItemDataAttributeValue& OutAttributeValue);
 
 	CONTENTBROWSERASSETDATASOURCE_API bool GetAssetFolderItemAttribute(const FContentBrowserAssetFolderItemDataPayload& InFolderPayload, const bool InIncludeMetaData, const FName InAttributeKey, FContentBrowserItemDataAttributeValue& OutAttributeValue);
 
diff --git a/Engine/Plugins/Editor/ContentBrowser/ContentBrowserAssetDataSource/Source/ContentBrowserAssetDataSource/Public/ContentBrowserAssetDataSource.h b/Engine/Plugins/Editor/ContentBrowser/ContentBrowserAssetDataSource/Source/ContentBrowserAssetDataSource/Public/ContentBrowserAssetDataSource.h
--- a/Engine/Plugins/Editor/ContentBrowser/ContentBrowserAssetDataSource/Source/ContentBrowserAssetDataSource/Public/ContentBrowserAssetDataSource.h
+++ b/Engine/Plugins/Editor/ContentBrowser/ContentBrowserAssetDataSource/Source/ContentBrowserAssetDataSource/Public/ContentBrowserAssetDataSource.h
@@ -27,6 +27,7 @@ class UToolMenu;
 class FAssetFolderContextMenu;
 class FAssetFileContextMenu;
 struct FCollectionNameType;
+class FContentBrowserModule;
 class UContentBrowserAssetDataSource;
 class UContentBrowserToolbarMenuContext;
 
@@ -461,9 +462,11 @@ private:
 
 	UE_API void RemoveRootContentPathFromStateMachine(const FString& InAssetPath);
 
-	IAssetRegistry* AssetRegistry;
+	IAssetRegistry* AssetRegistry = nullptr;
 
-	IAssetTools* AssetTools;
+	IAssetTools* AssetTools = nullptr;
+
+	FContentBrowserModule* ContentBrowserModule = nullptr;
 
 	ICollectionManager* CollectionManager;
 
