Exposed FullyLoadPackages/FullyLoadAssets functions to editor scripting via EditorLoadingAndSavingUtils

These fully load the given objects, so that UE no longer has a write lock on their on-disk files (eg, when needing to overwrite the file via an external application)

This is also available in the UI via Asset Actions -> Load, if you hold shift while clicking on Load

diff --git a/Engine/Plugins/Editor/ContentBrowser/ContentBrowserAssetDataSource/Source/ContentBrowserAssetDataSource/Private/AssetFileContextMenu.cpp b/Engine/Plugins/Editor/ContentBrowser/ContentBrowserAssetDataSource/Source/ContentBrowserAssetDataSource/Private/AssetFileContextMenu.cpp
--- a/Engine/Plugins/Editor/ContentBrowser/ContentBrowserAssetDataSource/Source/ContentBrowserAssetDataSource/Private/AssetFileContextMenu.cpp
+++ b/Engine/Plugins/Editor/ContentBrowser/ContentBrowserAssetDataSource/Source/ContentBrowserAssetDataSource/Private/AssetFileContextMenu.cpp
@@ -63,11 +63,10 @@
 #include "Subsystems/AssetEditorSubsystem.h"
 
 #include "Commandlets/TextAssetCommandlet.h"
+#include "FileHelpers.h"
 #include "Misc/FileHelper.h"
 #include "Internationalization/GatherableTextData.h"
 
-#include "ContentBrowserDataMenuContexts.h"
-
 #include "Editor/UnrealEdEngine.h"
 #include "Preferences/UnrealEdOptions.h"
 #include "UnrealEdGlobals.h"
@@ -461,7 +460,7 @@ void FAssetFileContextMenu::MakeAssetActionsSubMenu(UToolMenu* Menu)
 		Section.AddMenuEntry(
 			"Load",
 			LOCTEXT("Load", "Load"),
-			LOCTEXT("LoadTooltip", "Loads the selected assets into memory without opening any asset editors."),
+			LOCTEXT("LoadTooltip", "Loads the selected assets into memory without opening any asset editors.\nShift+Click to fully load the selected assets so that UE no longer has a write lock on their on-disk files."),
 			FSlateIcon(),
 			FUIAction(
 				FExecuteAction::CreateSP(this, &FAssetFileContextMenu::ExecuteLoad),
@@ -1685,14 +1684,25 @@ bool FAssetFileContextMenu::CanExecuteLoad() const
 
 void FAssetFileContextMenu::ExecuteLoad()
 {
+	const bool bFullyLoad = FSlateApplication::Get().GetModifierKeys().IsShiftDown();
+
 	FScopedSlowTask SlowTask(SelectedAssets.Num(), LOCTEXT("LoadingSelectedAssets", "Loading Selected Assets..."));
 	SlowTask.MakeDialogDelayed(1.0f);
 
+	TArray<UObject*> LoadedAssets;
 	for (const FAssetData& SelectedAsset : SelectedAssets)
 	{
-		SelectedAsset.GetAsset();
+		if (UObject* LoadedAsset = SelectedAsset.GetAsset())
+		{
+			LoadedAssets.Add(LoadedAsset);
+		}
 		SlowTask.EnterProgressFrame(1);
 	}
+
+	if (bFullyLoad)
+	{
+		UEditorLoadingAndSavingUtils::FullyLoadAssets(LoadedAssets);
+	}
 }
 
 bool FAssetFileContextMenu::CanExecuteReload() const
diff --git a/Engine/Source/Editor/UnrealEd/Private/FileHelpers.cpp b/Engine/Source/Editor/UnrealEd/Private/FileHelpers.cpp
--- a/Engine/Source/Editor/UnrealEd/Private/FileHelpers.cpp
+++ b/Engine/Source/Editor/UnrealEd/Private/FileHelpers.cpp
@@ -20,6 +20,7 @@
 #include "Misc/FileHelper.h"
 #include "Modules/ModuleManager.h"
 #include "UObject/Linker.h"
+#include "UObject/Package.h"
 #include "UObject/UObjectHash.h"
 #include "UObject/UObjectIterator.h"
 #include "ProfilingDebugging/ScopedTimers.h"
@@ -5969,4 +5970,28 @@ void UEditorLoadingAndSavingUtils::ReloadPackages(const TArray<UPackage*>& Packa
 	bOutAnyPackagesReloaded = UPackageTools::ReloadPackages(PackagesToReload, OutErrorMessage, InteractionMode);
 }
 
+void UEditorLoadingAndSavingUtils::FullyLoadPackages(const TArray<UPackage*>& Packages)
+{
+	TArray<int32> PendingPackageRequestIds;
+	for (UPackage* Package : Packages)
+	{
+		if (!Package->IsFullyLoaded())
+		{
+			PendingPackageRequestIds.Emplace(LoadPackageAsync(Package->GetName()));
+		}
+	}
+	if (PendingPackageRequestIds.Num() > 0)
+	{
+		FlushAsyncLoading(PendingPackageRequestIds);
+	}
+	ResetLoaders((const TArray<UObject*>&)Packages);
+}
+
+void UEditorLoadingAndSavingUtils::FullyLoadAssets(const TArray<UObject*>& Assets)
+{
+	TArray<UPackage*> Packages;
+	Algo::Transform(Assets, Packages, &UObject::GetPackage);
+	FullyLoadPackages(Packages);
+}
+
 #undef LOCTEXT_NAMESPACE
diff --git a/Engine/Source/Editor/UnrealEd/Public/FileHelpers.h b/Engine/Source/Editor/UnrealEd/Public/FileHelpers.h
--- a/Engine/Source/Editor/UnrealEd/Public/FileHelpers.h
+++ b/Engine/Source/Editor/UnrealEd/Public/FileHelpers.h
@@ -174,6 +174,22 @@ public:
 	 */
 	UFUNCTION(BlueprintCallable, Category = "Editor Scripting | Editor Loading and Saving")
 	static UNREALED_API void ReloadPackages(const TArray<UPackage*>& PackagesToReload, bool& bOutAnyPackagesReloaded, FText& OutErrorMessage, const EReloadPackagesInteractionMode InteractionMode = EReloadPackagesInteractionMode::Interactive);
+
+	/**
+	 * Fully loads the given packages, so that UE no longer has a write lock on their on-disk files.
+	 *
+	 * @param Packages Array of packages to fully load.
+	 */
+	UFUNCTION(BlueprintCallable, Category = "Editor Scripting | Editor Loading and Saving")
+	static UNREALED_API void FullyLoadPackages(const TArray<UPackage*>& Packages);
+
+	/**
+	 * Fully loads the given assets, so that UE no longer has a write lock on their on-disk files.
+	 *
+	 * @param Assets Array of assets to fully load.
+	 */
+	UFUNCTION(BlueprintCallable, Category = "Editor Scripting | Editor Loading and Saving")
+	static UNREALED_API void FullyLoadAssets(const TArray<UObject*>& Assets);
 };
 
 
