Fix check when module dependencies aren't loaded
- Both of these plugins depend on the AssetTools and ContentBrowser modules being loaded. If they are not loaded a check will fire.
- Despite declaring dependencies properly, the engine makes no effort to load dependencies before dependents. Project and plugin module are loaded in `FEngineLoop::PreInit` while built-in modules are not loaded until later in `FEngineLoop::Init`.
- The only reason these `GetModuleChecked` calls have ever worked is because some other module ends up explicitly loading the dependencies. As of CL 45081290 (which accesses the dependencies sooner) this led to a check when `DisableEnginePluginsByDefault` is enabled. `NiagaraEditor` doesn't load in this case, and was loading the `ContentBrowser` module before the plugins below were initialized.

diff --git a/Engine/Plugins/Editor/ContentBrowser/ContentBrowserAliasDataSource/Source/ContentBrowserAliasDataSource/Private/ContentBrowserAliasDataSource.cpp b/Engine/Plugins/Editor/ContentBrowser/ContentBrowserAliasDataSource/Source/ContentBrowserAliasDataSource/Private/ContentBrowserAliasDataSource.cpp
--- a/Engine/Plugins/Editor/ContentBrowser/ContentBrowserAliasDataSource/Source/ContentBrowserAliasDataSource/Private/ContentBrowserAliasDataSource.cpp
+++ b/Engine/Plugins/Editor/ContentBrowser/ContentBrowserAliasDataSource/Source/ContentBrowserAliasDataSource/Private/ContentBrowserAliasDataSource.cpp
@@ -73,8 +73,8 @@ void UContentBrowserAliasDataSource::Initialize(const bool InAutoRegister)
 	if (GIsEditor && !IsRunningCommandlet())
 	{
 		AssetRegistry = &FModuleManager::LoadModuleChecked<FAssetRegistryModule>(AssetRegistryConstants::ModuleName).Get();
-		AssetTools = &FModuleManager::GetModuleChecked<FAssetToolsModule>("AssetTools").Get();
-		ContentBrowserModule = &FModuleManager::Get().GetModuleChecked<FContentBrowserModule>(TEXT("ContentBrowser"));
+		AssetTools = &FModuleManager::LoadModuleChecked<FAssetToolsModule>("AssetTools").Get();
+		ContentBrowserModule = &FModuleManager::LoadModuleChecked<FContentBrowserModule>("ContentBrowser");
 
 		AssetRegistry->OnAssetAdded().AddUObject(this, &UContentBrowserAliasDataSource::OnAssetAdded);
 		AssetRegistry->OnAssetRemoved().AddUObject(this, &UContentBrowserAliasDataSource::OnAssetRemoved);
diff --git a/Engine/Plugins/Editor/ContentBrowser/ContentBrowserAssetDataSource/Source/ContentBrowserAssetDataSource/Private/ContentBrowserAssetDataSource.cpp b/Engine/Plugins/Editor/ContentBrowser/ContentBrowserAssetDataSource/Source/ContentBrowserAssetDataSource/Private/ContentBrowserAssetDataSource.cpp
--- a/Engine/Plugins/Editor/ContentBrowser/ContentBrowserAssetDataSource/Source/ContentBrowserAssetDataSource/Private/ContentBrowserAssetDataSource.cpp
+++ b/Engine/Plugins/Editor/ContentBrowser/ContentBrowserAssetDataSource/Source/ContentBrowserAssetDataSource/Private/ContentBrowserAssetDataSource.cpp
@@ -158,16 +158,8 @@ void UContentBrowserAssetDataSource::Initialize(const bool InAutoRegister)
 
 	AssetRegistry = &FModuleManager::LoadModuleChecked<FAssetRegistryModule>(AssetRegistryConstants::ModuleName).Get();
 	AssetRegistry->OnFileLoadProgressUpdated().AddUObject(this, &UContentBrowserAssetDataSource::OnAssetRegistryFileLoadProgress);
-
-	{
-		static const FName NAME_AssetTools = "AssetTools";
-		AssetTools = &FModuleManager::GetModuleChecked<FAssetToolsModule>(NAME_AssetTools).Get();
-	}
-
-	{
-		static const FName NAME_ContentBrowser = "ContentBrowser";
-		ContentBrowserModule = &FModuleManager::Get().GetModuleChecked<FContentBrowserModule>(NAME_ContentBrowser);
-	}
+	AssetTools = &FModuleManager::LoadModuleChecked<FAssetToolsModule>("AssetTools").Get();
+	ContentBrowserModule = &FModuleManager::LoadModuleChecked<FContentBrowserModule>("ContentBrowser");
 
 	CollectionManager = &FCollectionManagerModule::GetModule().Get();
 
