Fixed check renaming a folder in the Content Browser to a name that creates a package path longer than the maximum FName length.

diff --git a/Engine/Plugins/Editor/ContentBrowser/ContentBrowserAssetDataSource/Source/ContentBrowserAssetDataSource/Private/ContentBrowserAssetDataCore.cpp b/Engine/Plugins/Editor/ContentBrowser/ContentBrowserAssetDataSource/Source/ContentBrowserAssetDataSource/Private/ContentBrowserAssetDataCore.cpp
--- a/Engine/Plugins/Editor/ContentBrowser/ContentBrowserAssetDataSource/Source/ContentBrowserAssetDataSource/Private/ContentBrowserAssetDataCore.cpp
+++ b/Engine/Plugins/Editor/ContentBrowser/ContentBrowserAssetDataSource/Source/ContentBrowserAssetDataSource/Private/ContentBrowserAssetDataCore.cpp
@@ -1056,13 +1056,7 @@ bool CanRenameAssetFileItem(IAssetTools* InAssetTools, const FContentBrowserAsse
 		return false;
 	}
 
-	if (InNewName && InNewName->Len() >= NAME_SIZE)
-	{
-		SetOptionalErrorMessage(OutErrorMsg, LOCTEXT("Error_AssetNameTooLarge", "This asset name is too long. Please choose a shorter name."));
-		return false;
-	}
-
-	if (InNewName && (InIsTempoarary || InAssetPayload.GetAssetData().AssetName != FName(**InNewName))) // Deliberately ignore case here to allow case-only renames of existing assets
+	if (InNewName && (InIsTempoarary || InNewName->Len() >= NAME_SIZE || InAssetPayload.GetAssetData().AssetName != FName(**InNewName))) // Deliberately ignore case here to allow case-only renames of existing assets
 	{
 		const FString PackageName = InAssetPayload.GetAssetData().PackagePath.ToString() / *InNewName;
 		const FString ObjectPath = FString::Printf(TEXT("%s.%s"), *PackageName, **InNewName);
diff --git a/Engine/Plugins/Editor/ContentBrowser/ContentBrowserAssetDataSource/Source/ContentBrowserAssetDataSource/Private/ContentBrowserAssetDataSource.cpp b/Engine/Plugins/Editor/ContentBrowser/ContentBrowserAssetDataSource/Source/ContentBrowserAssetDataSource/Private/ContentBrowserAssetDataSource.cpp
--- a/Engine/Plugins/Editor/ContentBrowser/ContentBrowserAssetDataSource/Source/ContentBrowserAssetDataSource/Private/ContentBrowserAssetDataSource.cpp
+++ b/Engine/Plugins/Editor/ContentBrowser/ContentBrowserAssetDataSource/Source/ContentBrowserAssetDataSource/Private/ContentBrowserAssetDataSource.cpp
@@ -2889,14 +2889,17 @@ bool UContentBrowserAssetDataSource::CanRenameItem(const FContentBrowserItemData
 		if (TSharedPtr<const FContentBrowserAssetFolderItemDataPayload> FolderPayload = GetAssetFolderItemPayload(InItem))
 		{
 			FString NewInternalPathString = FPaths::GetPath(FolderPayload->GetInternalPath().ToString()) / *InNewName;
-			FName NewInternalPath(NewInternalPathString);
-
-			const EContentBrowserFolderAttributes FolderAttributes = GetAssetFolderAttributes(NewInternalPath);
-			if (!EnumHasAnyFlags(FolderAttributes, EContentBrowserFolderAttributes::AlwaysVisible | EContentBrowserFolderAttributes::HasAssets | EContentBrowserFolderAttributes::HasRedirectors) &&
-				HideFolderIfEmpty(*HideFolderIfEmptyFilter, NewInternalPath, NewInternalPathString))
+			if (NewInternalPathString.Len() < NAME_SIZE) // If the path is too long, CanRenameItem will fail anyways, let it produce the error message.
 			{
-				// We are renaming to an existing hidden folder name, disable the unique name check since it will fail.
-				bCheckUniqueName = false;
+				FName NewInternalPath(NewInternalPathString);
+
+				const EContentBrowserFolderAttributes FolderAttributes = GetAssetFolderAttributes(NewInternalPath);
+				if (!EnumHasAnyFlags(FolderAttributes, EContentBrowserFolderAttributes::AlwaysVisible | EContentBrowserFolderAttributes::HasAssets | EContentBrowserFolderAttributes::HasRedirectors) &&
+					HideFolderIfEmpty(*HideFolderIfEmptyFilter, NewInternalPath, NewInternalPathString))
+				{
+					// We are renaming to an existing hidden folder name, disable the unique name check since it will fail.
+					bCheckUniqueName = false;
+				}
 			}
 		}
 	}
diff --git a/Engine/Plugins/Experimental/UAF/UAF/Source/UAFEditor/Private/Variables/Outliner/VariablesOutlinerAssetItem.cpp b/Engine/Plugins/Experimental/UAF/UAF/Source/UAFEditor/Private/Variables/Outliner/VariablesOutlinerAssetItem.cpp
--- a/Engine/Plugins/Experimental/UAF/UAF/Source/UAFEditor/Private/Variables/Outliner/VariablesOutlinerAssetItem.cpp
+++ b/Engine/Plugins/Experimental/UAF/UAF/Source/UAFEditor/Private/Variables/Outliner/VariablesOutlinerAssetItem.cpp
@@ -630,14 +630,8 @@ bool FVariablesOutlinerAssetItem::ValidateName(const FText& InNewName, FText& Ou
 		return false;
 	}
 
-	FString NewName = InNewName.ToString();
-	if (NewName.Len() >= NAME_SIZE)
-	{
-		OutErrorMessage = LOCTEXT("Error_AssetNameTooLarge", "This asset name is too long. Please choose a shorter name.");
-		return false;
-	}
-
-	if (Asset->GetFName() != FName(*NewName)) // Deliberately ignore case here to allow case-only renames of existing assets
+	const FString NewName = InNewName.ToString();
+	if (NewName.Len() >= NAME_SIZE || Asset->GetFName() != FName(*NewName)) // Deliberately ignore case here to allow case-only renames of existing assets
 	{
 		const FString PackageName = Asset->GetPackage()->GetPathName() / NewName;
 		const FString ObjectPath = FString::Printf(TEXT("%s.%s"), *PackageName, *NewName);
diff --git a/Engine/Source/Developer/AssetTools/Private/AssetViewUtils.cpp b/Engine/Source/Developer/AssetTools/Private/AssetViewUtils.cpp
--- a/Engine/Source/Developer/AssetTools/Private/AssetViewUtils.cpp
+++ b/Engine/Source/Developer/AssetTools/Private/AssetViewUtils.cpp
@@ -1454,7 +1454,8 @@ bool AssetViewUtils::IsValidObjectPathForCreate(const FString& ObjectPath, const
 	// Make sure we are not creating an FName that is too large
 	if ( ObjectPath.Len() >= NAME_SIZE )
 	{
-		OutErrorMessage = LOCTEXT("AssetNameTooLong", "This asset name is too long. Please choose a shorter name.");
+		OutErrorMessage = FText::Format(LOCTEXT("AssetNameTooLong", "The name for the asset is too long, the maximum is '{0}'. \nPlease choose a shorter name for the asset or create it in a shallower folder structure."),
+			FText::AsNumber(NAME_SIZE));
 		// Return false to indicate that the user should enter a new name
 		return false;
 	}
@@ -1539,6 +1540,15 @@ bool AssetViewUtils::IsValidFolderPathForCreate(const FString& InFolderPath, con
 		return false;
 	}
 
+	// Make sure we are not creating an FName that is too large
+	if (NewFolderPath.Len() >= NAME_SIZE)
+	{
+		OutErrorMessage = FText::Format(LOCTEXT("PackagePathTooLong", "The path for the folder is too deep, the maximum is '{0}'. \nPlease choose a shorter name for the folder or create it in a shallower folder structure."),
+			FText::AsNumber(NAME_SIZE));
+		// Return false to indicate that the user should enter a new name
+		return false;
+	}
+
 	// Make sure we are not creating a folder path that is too long
 	if (NewFolderPath.Len() > GetMaxCookPathLen() - MAX_CLASS_NAME_LENGTH)
 	{
