Fixed Content Browser Async factory process
- Previously, the Content Browser Async path would call the regular AssetTools completion path instead of the Content Browser completion path that will prompt for rename and defer the factory creation process until after rename.
- Renamed delegate to be more accurate

diff --git a/Engine/Plugins/Editor/ContentBrowser/ContentBrowserAssetDataSource/Source/ContentBrowserAssetDataSource/Private/ContentBrowserAssetDataSource.cpp b/Engine/Plugins/Editor/ContentBrowser/ContentBrowserAssetDataSource/Source/ContentBrowserAssetDataSource/Private/ContentBrowserAssetDataSource.cpp
--- a/Engine/Plugins/Editor/ContentBrowser/ContentBrowserAssetDataSource/Source/ContentBrowserAssetDataSource/Private/ContentBrowserAssetDataSource.cpp
+++ b/Engine/Plugins/Editor/ContentBrowser/ContentBrowserAssetDataSource/Source/ContentBrowserAssetDataSource/Private/ContentBrowserAssetDataSource.cpp
@@ -4019,28 +4019,8 @@ void UContentBrowserAssetDataSource::OnBeginCreateAsset(const FName InDefaultAss
 
 	FContentBrowserItemPath AssetPathToUse = ContentBrowserModule->Get().GetInitialPathToSaveAsset(FContentBrowserItemPath(InPackagePath, EContentBrowserPathType::Internal));
 
-	const bool bUseAsyncCreate = InFactory && EnumHasAnyFlags(InFactory->GetSupportedWorkflows(), EFactoryCreateWorkflow::Asynchronous);
-	const bool bShowDialogToPickPath = !AssetPathToUse.HasInternalPath() || (AssetPathToUse.GetInternalPathName() != InPackagePath);
-	if (bShowDialogToPickPath)
-	{
-		FAssetToolsModule& AssetToolsModule = FModuleManager::Get().LoadModuleChecked<FAssetToolsModule>("AssetTools");
-		const FString InitialInternalPath = AssetPathToUse.HasInternalPath() ? AssetPathToUse.GetInternalPathString() : TEXT("/Game");
-
-		if (bUseAsyncCreate)
-		{
-			AssetToolsModule.Get().CreateAssetWithDialogAsync(InDefaultAssetName.ToString(), InitialInternalPath, ClassToUse, InFactory, FAssetCreateComplete(), FAssetCreateCancelled(), NAME_None);
-		}
-		else
-		{
-			AssetToolsModule.Get().CreateAssetWithDialog(InDefaultAssetName.ToString(), InitialInternalPath, ClassToUse, InFactory, NAME_None, /*bCallConfigureProperties*/ false);
-		}
-	}
-	else if (bUseAsyncCreate)
-	{
-		FAssetToolsModule& AssetToolsModule = FModuleManager::Get().LoadModuleChecked<FAssetToolsModule>("AssetTools");
-		AssetToolsModule.Get().CreateAssetAsync(InDefaultAssetName.ToString(), InPackagePath.ToString(), ClassToUse, InFactory);
-	}
-	else
+	// This will create a content browser item, which will trigger the factory to create the actual object upon rename completion
+	auto CreateContentBrowserItemWithCreationPayload = [this, InDefaultAssetName, InPackagePath, InAssetClass, InFactory, ClassToUse, InOnBeginItemCreation]()
 	{
 		FAssetData NewAssetData(*(InPackagePath.ToString() / InDefaultAssetName.ToString()), InPackagePath, InDefaultAssetName, ClassToUse->GetClassPathName());
 
@@ -4066,6 +4046,38 @@ void UContentBrowserAssetDataSource::OnBeginCreateAsset(const FName InDefaultAss
 				this, &UContentBrowserAssetDataSource::OnValidateItemName),
 			FContentBrowserItemDataTemporaryContext::FOnFinalizeItem::CreateUObject(
 				this, &UContentBrowserAssetDataSource::OnFinalizeCreateAsset)));
+	};
+
+	const bool bUseAsyncCreate = InFactory && EnumHasAnyFlags(InFactory->GetSupportedWorkflows(), EFactoryCreateWorkflow::Asynchronous);
+	const bool bShowDialogToPickPath = !AssetPathToUse.HasInternalPath() || (AssetPathToUse.GetInternalPathName() != InPackagePath);
+	if (bShowDialogToPickPath)
+	{
+		FAssetToolsModule& AssetToolsModule = FModuleManager::Get().LoadModuleChecked<FAssetToolsModule>("AssetTools");
+		const FString InitialInternalPath = AssetPathToUse.HasInternalPath() ? AssetPathToUse.GetInternalPathString() : TEXT("/Game");
+
+		if (bUseAsyncCreate)
+		{
+			AssetToolsModule.Get().CreateAssetWithDialogAsync(InDefaultAssetName.ToString(), InitialInternalPath, ClassToUse, InFactory, FAssetCreateComplete(), FAssetCreateCancelled(), NAME_None);
+		}
+		else
+		{
+			AssetToolsModule.Get().CreateAssetWithDialog(InDefaultAssetName.ToString(), InitialInternalPath, ClassToUse, InFactory, NAME_None, /*bCallConfigureProperties*/ false);
+		}
+	}
+	else if (bUseAsyncCreate)
+	{
+		FOnFactoryConfigurePropertiesAsyncComplete OnConfigurePropertiesComplete = FOnFactoryConfigurePropertiesAsyncComplete::CreateLambda([CreateContentBrowserItemWithCreationPayload](UFactory* Factory)
+		{
+			CreateContentBrowserItemWithCreationPayload();
+		});
+
+		// At this point, the configuration process is completely handed off to the factory. It's the factory's responsibility to call Complete or Cancel.
+		// Complete will call the regular content browser creation code path.
+		InFactory->ConfigurePropertiesAsync(OnConfigurePropertiesComplete, FOnFactoryConfigurePropertiesAsyncCancelled());
+	}
+	else
+	{
+		CreateContentBrowserItemWithCreationPayload();
 	}
 }
 
diff --git a/Engine/Source/Developer/AssetTools/Private/AssetTools.cpp b/Engine/Source/Developer/AssetTools/Private/AssetTools.cpp
--- a/Engine/Source/Developer/AssetTools/Private/AssetTools.cpp
+++ b/Engine/Source/Developer/AssetTools/Private/AssetTools.cpp
@@ -1781,8 +1781,7 @@ UObject* UAssetToolsImpl::CreateAsset(const FString& AssetName, const FString& P
 	return NewObj;
 }
 
-bool UAssetToolsImpl::CreateAssetAsync(const FString& AssetName, const FString& PackagePath, UClass* AssetClass, UFactory* Factory,
-	FAssetCreateComplete OnComplete, FAssetCreateCancelled OnCancelled, FName CallingContext)
+bool UAssetToolsImpl::CreateAssetAsync(const FString& AssetName, const FString& PackagePath, UClass* AssetClass, UFactory* Factory, FAssetCreateComplete OnComplete, FAssetCreateCancelled OnCancelled, FName CallingContext)
 {
 	// Verify the factory class
 	if (!ensure(AssetClass && Factory))
@@ -1827,9 +1826,9 @@ bool UAssetToolsImpl::CreateAssetAsync(const FString& AssetName, const FString&
 	TStrongObjectPtr<UPackage> DontGCPackage(Pkg);
 
 	return Factory->ConfigurePropertiesAsync(
-		FFactoryConfigurePropertiesComplete::CreateUObject(this, &UAssetToolsImpl::CreateAssetAsync_OnComplete, DontGCPackage, ClassToUse, 
+		FOnFactoryConfigurePropertiesAsyncComplete::CreateUObject(this, &UAssetToolsImpl::CreateAssetAsync_OnComplete, DontGCPackage, ClassToUse, 
 			AssetName, Flags, CallingContext, PackageName, OnComplete, OnCancelled),
-		FFactoryConfigurePropertiesCancelled()
+		FOnFactoryConfigurePropertiesAsyncCancelled()
 	);
 }
 
@@ -2078,9 +2077,9 @@ bool UAssetToolsImpl::CreateAssetWithDialogAsync(const FString& AssetName, const
 	FEditorDelegates::OnConfigureNewAssetProperties.Broadcast(Factory);
 
 	return Factory->ConfigurePropertiesAsync(
-		FFactoryConfigurePropertiesComplete::CreateUObject(this, &UAssetToolsImpl::CreateAssetWithDialogAsync_OnComplete,
+		FOnFactoryConfigurePropertiesAsyncComplete::CreateUObject(this, &UAssetToolsImpl::CreateAssetWithDialogAsync_OnComplete,
 			AssetClass, CallingContext, SaveObjectPath, OnComplete, OnCancelled),
-		FFactoryConfigurePropertiesCancelled()
+		FOnFactoryConfigurePropertiesAsyncCancelled()
 	);
 }
 
diff --git a/Engine/Source/Editor/UnrealEd/Classes/Factories/Factory.h b/Engine/Source/Editor/UnrealEd/Classes/Factories/Factory.h
--- a/Engine/Source/Editor/UnrealEd/Classes/Factories/Factory.h
+++ b/Engine/Source/Editor/UnrealEd/Classes/Factories/Factory.h
@@ -33,8 +33,8 @@ enum class EFactoryCreateWorkflow : uint8
 };
 ENUM_CLASS_FLAGS(EFactoryCreateWorkflow);
 
-DECLARE_DELEGATE_OneParam(FFactoryConfigurePropertiesComplete, UFactory* /* this */);
-DECLARE_DELEGATE_OneParam(FFactoryConfigurePropertiesCancelled, UFactory* /* this */);
+DECLARE_DELEGATE_OneParam(FOnFactoryConfigurePropertiesAsyncComplete, UFactory* /* this */);
+DECLARE_DELEGATE_OneParam(FOnFactoryConfigurePropertiesAsyncCancelled, UFactory* /* this */);
 
 /**
  * Base class for all factories
@@ -208,7 +208,7 @@ public:
 	 * @param OnCancelled Called when the asynchronous process was cancelled.
 	 * @return true if the async process was successfully started.
 	 */
-	virtual bool ConfigurePropertiesAsync(FFactoryConfigurePropertiesComplete OnComplete, FFactoryConfigurePropertiesCancelled OnCancelled)
+	virtual bool ConfigurePropertiesAsync(FOnFactoryConfigurePropertiesAsyncComplete OnComplete, FOnFactoryConfigurePropertiesAsyncCancelled OnCancelled)
 	{
 		OnCancelled.Execute(this);
 		return false;
