Fix content browser favorites being lost when plugins unmounted
[RN] Fix content browser favorites being lost when unmounting plugins
#Synced-CL 48801413

diff --git a/Engine/Plugins/Editor/ContentBrowser/ContentBrowserAssetDataSource/Source/ContentBrowserAssetDataSource/Private/ContentBrowserAssetDataSource.cpp b/Engine/Plugins/Editor/ContentBrowser/ContentBrowserAssetDataSource/Source/ContentBrowserAssetDataSource/Private/ContentBrowserAssetDataSource.cpp
--- a/Engine/Plugins/Editor/ContentBrowser/ContentBrowserAssetDataSource/Source/ContentBrowserAssetDataSource/Private/ContentBrowserAssetDataSource.cpp
+++ b/Engine/Plugins/Editor/ContentBrowser/ContentBrowserAssetDataSource/Source/ContentBrowserAssetDataSource/Private/ContentBrowserAssetDataSource.cpp
@@ -3591,7 +3591,16 @@ void UContentBrowserAssetDataSource::OnPathsRemoved(TConstArrayView<FStringView>
 		RecentlyPopulatedAssetFolders.Remove(PathName);
 		AssetFolderToAttributes.Remove(PathName);
 
-		QueueItemDataUpdate(FContentBrowserItemDataUpdate::MakeItemRemovedUpdate(CreateAssetFolderItem(PathName)));
+		const FStringView MountPointWithLeadingSlash = FPathViews::GetMountPointNameFromPath(InPath, nullptr, false);
+		if (FPackageName::MountPointExists(MountPointWithLeadingSlash))
+		{
+			QueueItemDataUpdate(FContentBrowserItemDataUpdate::MakeItemRemovedUpdate(CreateAssetFolderItem(PathName)));
+		}
+		else
+		{
+			// Required to avoid loss of favorites for plugins being unmounted
+			QueueItemDataUpdate(FContentBrowserItemDataUpdate::MakeItemUnmountedUpdate(CreateAssetFolderItem(PathName)));
+		}
 
 		OnAssetPathRemovedDelegate.Broadcast(PathName, GetTypeHash(PathName));
 	}
diff --git a/Engine/Source/Editor/ContentBrowser/Private/ContentBrowserUtils.cpp b/Engine/Source/Editor/ContentBrowser/Private/ContentBrowserUtils.cpp
--- a/Engine/Source/Editor/ContentBrowser/Private/ContentBrowserUtils.cpp
+++ b/Engine/Source/Editor/ContentBrowser/Private/ContentBrowserUtils.cpp
@@ -1148,18 +1148,26 @@ void ContentBrowserUtils::RemoveFavoriteFolder(const FContentBrowserItemPath& Fo
 		return;
 	}
 
-	FString InvariantFolder = InvariantPath.ToString();
-
-	FContentBrowserSingleton::Get().FavoriteFolderPaths.Remove(InvariantFolder);
+	const FString InvariantFolder = InvariantPath.ToString();
+	bool bModified = false;
+	if (FContentBrowserSingleton::Get().FavoriteFolderPaths.Remove(InvariantFolder) > 0)
+	{
+		bModified = true;
+	}
 
 	if (UContentBrowserConfig* EditorConfig = UContentBrowserConfig::Get())
 	{
-		EditorConfig->Favorites.Remove(InvariantFolder);
-
-		UContentBrowserConfig::Get()->SaveEditorConfig();
+		if (EditorConfig->Favorites.Remove(InvariantFolder) > 0)
+		{
+			bModified = true;
+			UContentBrowserConfig::Get()->SaveEditorConfig();
+		}
 	}
 
-	FContentBrowserSingleton::Get().BroadcastFavoritesChanged(FolderPath, false);
+	if (bModified)
+	{
+		FContentBrowserSingleton::Get().BroadcastFavoritesChanged(FolderPath, false);
+	}
 }
 
 void ContentBrowserUtils::RemoveFavoriteFolder(const FString& FolderPath, bool bFlushConfig)
diff --git a/Engine/Source/Editor/ContentBrowser/Private/SAssetView.cpp b/Engine/Source/Editor/ContentBrowser/Private/SAssetView.cpp
--- a/Engine/Source/Editor/ContentBrowser/Private/SAssetView.cpp
+++ b/Engine/Source/Editor/ContentBrowser/Private/SAssetView.cpp
@@ -6819,6 +6819,10 @@ void SAssetView::HandleItemDataUpdated(TArrayView<const FContentBrowserItemDataU
 			RemoveItem(FContentBrowserMinimalItemData(ItemData));
 			break;
 
+		case EContentBrowserItemUpdateType::Unmounted:
+			RemoveItem(FContentBrowserMinimalItemData(ItemData));
+			break;
+
 		default:
 			checkf(false, TEXT("Unexpected EContentBrowserItemUpdateType!"));
 			break;
diff --git a/Engine/Source/Editor/ContentBrowser/Private/SContentBrowser.cpp b/Engine/Source/Editor/ContentBrowser/Private/SContentBrowser.cpp
--- a/Engine/Source/Editor/ContentBrowser/Private/SContentBrowser.cpp
+++ b/Engine/Source/Editor/ContentBrowser/Private/SContentBrowser.cpp
@@ -4751,6 +4751,10 @@ void SContentBrowser::HandleItemDataUpdated(TArrayView<const FContentBrowserItem
 			HandlePathRemoved(ItemDataUpdate.GetItemData().GetVirtualPath());
 			break;
 
+		case EContentBrowserItemUpdateType::Unmounted:
+			HandlePathRemoved(ItemDataUpdate.GetItemData().GetVirtualPath());
+			break;
+
 		default:
 			break;
 		}
diff --git a/Engine/Source/Editor/ContentBrowser/Private/SPathView.cpp b/Engine/Source/Editor/ContentBrowser/Private/SPathView.cpp
--- a/Engine/Source/Editor/ContentBrowser/Private/SPathView.cpp
+++ b/Engine/Source/Editor/ContentBrowser/Private/SPathView.cpp
@@ -594,6 +594,13 @@ void FPathViewData::ProcessDataUpdates(TConstArrayView<FContentBrowserItemDataUp
 				TryRemoveFolderItemInternal(ItemData);
 				break;
 
+			case EContentBrowserItemUpdateType::Unmounted:
+				UE_LOG(LogPathView, VeryVerbose, TEXT("[%s] Unmounted item %s:%s"), *WriteToString<256>(OwningContentBrowserName),
+					*WriteToString<256>(ItemData.GetOwnerDataSource()->GetFName()),
+					*WriteToString<256>(ItemData.GetVirtualPath()));
+				TryRemoveFolderItemInternal(ItemData);
+				break;
+
 			default:
 				checkf(false, TEXT("Unexpected EContentBrowserItemUpdateType!"));
 				break;
@@ -3162,6 +3169,9 @@ void SFavoritePathView::HandleItemDataUpdated(TArrayView<const FContentBrowserIt
 				ContentBrowserUtils::RemoveFavoriteFolder(
 					FContentBrowserItemPath(ItemData.GetVirtualPath(), EContentBrowserPathType::Virtual));
 				break;
+			case EContentBrowserItemUpdateType::Unmounted:
+				// Do not remove favorite because a plugin can be unmounted then mounted again later
+				break;
 			default:
 				checkf(false, TEXT("Unexpected EContentBrowserItemUpdateType!"));
 				break;
diff --git a/Engine/Source/Editor/ContentBrowserData/Private/ContentBrowserItemData.cpp b/Engine/Source/Editor/ContentBrowserData/Private/ContentBrowserItemData.cpp
--- a/Engine/Source/Editor/ContentBrowserData/Private/ContentBrowserItemData.cpp
+++ b/Engine/Source/Editor/ContentBrowserData/Private/ContentBrowserItemData.cpp
@@ -424,6 +424,14 @@ FContentBrowserItemDataUpdate FContentBrowserItemDataUpdate::MakeItemRemovedUpda
 	return ItemUpdate;
 }
 
+FContentBrowserItemDataUpdate FContentBrowserItemDataUpdate::MakeItemUnmountedUpdate(FContentBrowserItemData InItemData)
+{
+	FContentBrowserItemDataUpdate ItemUpdate;
+	ItemUpdate.UpdateType = EContentBrowserItemUpdateType::Unmounted;
+	ItemUpdate.ItemData = MoveTemp(InItemData);
+	return ItemUpdate;
+}
+
 EContentBrowserItemUpdateType FContentBrowserItemDataUpdate::GetUpdateType() const
 {
 	return UpdateType;
diff --git a/Engine/Source/Editor/ContentBrowserData/Public/ContentBrowserItemData.h b/Engine/Source/Editor/ContentBrowserData/Public/ContentBrowserItemData.h
--- a/Engine/Source/Editor/ContentBrowserData/Public/ContentBrowserItemData.h
+++ b/Engine/Source/Editor/ContentBrowserData/Public/ContentBrowserItemData.h
@@ -91,6 +91,8 @@ enum class EContentBrowserItemUpdateType : uint8
 	Moved,
 	/** This item existed and has been deleted */
 	Removed,
+	/** This item existed and has been unmounted */
+	Unmounted,
 };
 
 /**
@@ -786,6 +788,11 @@ public:
 	 */
 	static UE_API FContentBrowserItemDataUpdate MakeItemRemovedUpdate(FContentBrowserItemData InItemData);
 
+	/**
+	 * Create an item update for when an item is unmounted.
+	 */
+	static UE_API FContentBrowserItemDataUpdate MakeItemUnmountedUpdate(FContentBrowserItemData InItemData);
+
 	/**
 	 * Get the value denoting the types of updates that can be emitted for an item.
 	 */
