Enable AutoRTFM on editor modules that register on `OnObjectPropertyChanged`.

#Synced-CL 49758756

diff --git a/Engine/Plugins/2D/Paper2D/Source/Paper2DEditor/Paper2DEditor.Build.cs b/Engine/Plugins/2D/Paper2D/Source/Paper2DEditor/Paper2DEditor.Build.cs
--- a/Engine/Plugins/2D/Paper2D/Source/Paper2DEditor/Paper2DEditor.Build.cs
+++ b/Engine/Plugins/2D/Paper2D/Source/Paper2DEditor/Paper2DEditor.Build.cs
@@ -44,5 +44,9 @@ public class Paper2DEditor : ModuleRules
 				"Settings",
 				"LevelEditor"
 			});
+
+		// This module needs AutoRTFM so disable the auto disable since its in a plugin.
+		// The call to the delegate `OnObjectPropertyChanged` calls into here.
+		bDisableAutoRTFMInstrumentation = false;
 	}
 }
diff --git a/Engine/Plugins/Editor/ContentBrowser/ContentBrowserAliasDataSource/Source/ContentBrowserAliasDataSource/ContentBrowserAliasDataSource.Build.cs b/Engine/Plugins/Editor/ContentBrowser/ContentBrowserAliasDataSource/Source/ContentBrowserAliasDataSource/ContentBrowserAliasDataSource.Build.cs
--- a/Engine/Plugins/Editor/ContentBrowser/ContentBrowserAliasDataSource/Source/ContentBrowserAliasDataSource/ContentBrowserAliasDataSource.Build.cs
+++ b/Engine/Plugins/Editor/ContentBrowser/ContentBrowserAliasDataSource/Source/ContentBrowserAliasDataSource/ContentBrowserAliasDataSource.Build.cs
@@ -25,6 +25,10 @@ namespace UnrealBuildTool.Rules
 					"Engine"
 				}
 			);
+
+			// This module needs AutoRTFM so disable the auto disable since its in a plugin.
+			// The call to the delegate `OnObjectPropertyChanged` calls into here.
+			bDisableAutoRTFMInstrumentation = false;
 		}
 	}
 }
diff --git a/Engine/Plugins/Editor/ContentBrowser/ContentBrowserAssetDataSource/Source/ContentBrowserAssetDataSource/ContentBrowserAssetDataSource.Build.cs b/Engine/Plugins/Editor/ContentBrowser/ContentBrowserAssetDataSource/Source/ContentBrowserAssetDataSource/ContentBrowserAssetDataSource.Build.cs
--- a/Engine/Plugins/Editor/ContentBrowser/ContentBrowserAssetDataSource/Source/ContentBrowserAssetDataSource/ContentBrowserAssetDataSource.Build.cs
+++ b/Engine/Plugins/Editor/ContentBrowser/ContentBrowserAssetDataSource/Source/ContentBrowserAssetDataSource/ContentBrowserAssetDataSource.Build.cs
@@ -47,6 +47,10 @@ namespace UnrealBuildTool.Rules
 					"SourceControlWindows",
 				}
 			);
+
+			// This module needs AutoRTFM so disable the auto disable since its in a plugin.
+			// The call to the delegate `OnObjectPropertyChanged` calls into here.
+			bDisableAutoRTFMInstrumentation = false;
 		}
 	}
 }
diff --git a/Engine/Source/Editor/DataLayerEditor/DataLayerEditor.Build.cs b/Engine/Source/Editor/DataLayerEditor/DataLayerEditor.Build.cs
--- a/Engine/Source/Editor/DataLayerEditor/DataLayerEditor.Build.cs
+++ b/Engine/Source/Editor/DataLayerEditor/DataLayerEditor.Build.cs
@@ -26,5 +26,9 @@ public class DataLayerEditor : ModuleRules
 				"ContentBrowserData"
 			}
 		);
+
+		// This module needs AutoRTFM so disable the auto disable since its in a plugin.
+		// The call to the delegate `OnObjectPropertyChanged` calls into here.
+		bDisableAutoRTFMInstrumentation = false;
 	}
 }
diff --git a/Engine/Source/Editor/LevelEditor/LevelEditor.Build.cs b/Engine/Source/Editor/LevelEditor/LevelEditor.Build.cs
--- a/Engine/Source/Editor/LevelEditor/LevelEditor.Build.cs
+++ b/Engine/Source/Editor/LevelEditor/LevelEditor.Build.cs
@@ -127,10 +127,13 @@ public class LevelEditor : ModuleRules
 			DynamicallyLoadedModuleNames.Add("SessionFrontend");
 		}
 
-
 		if (Target.bWithLiveCoding)
 		{
 			PrivateIncludePathModuleNames.Add("LiveCoding");
 		}
+
+		// This module needs AutoRTFM so disable the auto disable since its in a plugin.
+		// The call to the delegate `OnObjectPropertyChanged` calls into here.
+		bDisableAutoRTFMInstrumentation = false;
 	}
 }
