Enabling PropertyMatrix for UMaterial assets.

[RN] major | Rendering - Materials and Shaders

#Synced-CL 49784419

diff --git a/Engine/Plugins/Editor/ContentBrowser/ContentBrowserAssetDataSource/Source/ContentBrowserAssetDataSource/Private/AssetFileContextMenu.cpp b/Engine/Plugins/Editor/ContentBrowser/ContentBrowserAssetDataSource/Source/ContentBrowserAssetDataSource/Private/AssetFileContextMenu.cpp
--- a/Engine/Plugins/Editor/ContentBrowser/ContentBrowserAssetDataSource/Source/ContentBrowserAssetDataSource/Private/AssetFileContextMenu.cpp
+++ b/Engine/Plugins/Editor/ContentBrowser/ContentBrowserAssetDataSource/Source/ContentBrowserAssetDataSource/Private/AssetFileContextMenu.cpp
@@ -489,7 +489,7 @@ void FAssetFileContextMenu::MakeAssetActionsSubMenu(UToolMenu* Menu)
 			// Materials can't be bulk edited currently as they require very special handling because of their dependencies with the rendering thread, and we'd have to hack the property matrix too much.
 			for (const FAssetData& Asset : SelectedAssets)
 			{
-				if (Asset.AssetClassPath == UMaterial::StaticClass()->GetClassPathName() || Asset.AssetClassPath == UMaterialInstanceConstant::StaticClass()->GetClassPathName() || Asset.AssetClassPath == UMaterialFunction::StaticClass()->GetClassPathName() || Asset.AssetClassPath == UMaterialFunctionInstance::StaticClass()->GetClassPathName())
+				if (Asset.AssetClassPath == UMaterialInstanceConstant::StaticClass()->GetClassPathName() || Asset.AssetClassPath == UMaterialFunction::StaticClass()->GetClassPathName() || Asset.AssetClassPath == UMaterialFunctionInstance::StaticClass()->GetClassPathName())
 				{
 					bCanUsePropertyMatrix = false;
 					break;
