EditorUsability : ContentBrowser
- Added new ContentBrowserAssetDataSource to handle new commands
- Added command with keyboard shortcut to Capture Thumbnail for asset in the Content Browser (CTRL+SHIFT+T)

#Synced-CL 52060601

diff --git a/Engine/Plugins/Editor/ContentBrowser/ContentBrowserAssetDataSource/Source/ContentBrowserAssetDataSource/Private/AssetFileContextMenu.cpp b/Engine/Plugins/Editor/ContentBrowser/ContentBrowserAssetDataSource/Source/ContentBrowserAssetDataSource/Private/AssetFileContextMenu.cpp
--- a/Engine/Plugins/Editor/ContentBrowser/ContentBrowserAssetDataSource/Source/ContentBrowserAssetDataSource/Private/AssetFileContextMenu.cpp
+++ b/Engine/Plugins/Editor/ContentBrowser/ContentBrowserAssetDataSource/Source/ContentBrowserAssetDataSource/Private/AssetFileContextMenu.cpp
@@ -46,6 +46,7 @@
 
 #include "Kismet2/KismetEditorUtilities.h"
 #include "ComponentAssetBroker.h"
+#include "ContentBrowserAssetDataSourceCommands.h"
 #include "Widgets/Input/SNumericEntryBox.h"
 
 #include "SourceCodeNavigation.h"
@@ -340,19 +341,10 @@ void FAssetFileContextMenu::MakeAssetActionsSubMenu(UToolMenu* Menu)
 		}
 
 		// Capture Thumbnail
-		FAssetToolsModule& AssetToolsModule = FModuleManager::LoadModuleChecked<FAssetToolsModule>("AssetTools");
-		if (bCanBeModified && SelectedAssets.Num() == 1 && AssetToolsModule.Get().AssetUsesGenericThumbnail(SelectedAssets[0]))
+		if (bCanBeModified && !SelectedAssets.IsEmpty())
 		{
-			Section.AddMenuEntry(
-				"CaptureThumbnail",
-				LOCTEXT("CaptureThumbnail", "Capture Thumbnail"),
-				LOCTEXT("CaptureThumbnailTooltip", "Captures a thumbnail from the active viewport."),
-				FSlateIcon(),
-				FUIAction(
-					FExecuteAction::CreateSP(this, &FAssetFileContextMenu::ExecuteCaptureThumbnail),
-					FCanExecuteAction::CreateSP(this, &FAssetFileContextMenu::CanExecuteCaptureThumbnail)
-				)
-			);
+			const FContentBrowserAssetDataSourceCommands& AssetDataSourceCommands = FContentBrowserAssetDataSourceCommands::Get();
+			Section.AddMenuEntry(AssetDataSourceCommands.CaptureThumbnail);
 		}
 
 		// Clear Thumbnail
diff --git a/Engine/Plugins/Editor/ContentBrowser/ContentBrowserAssetDataSource/Source/ContentBrowserAssetDataSource/Private/ContentBrowserAssetDataSource.cpp b/Engine/Plugins/Editor/ContentBrowser/ContentBrowserAssetDataSource/Source/ContentBrowserAssetDataSource/Private/ContentBrowserAssetDataSource.cpp
--- a/Engine/Plugins/Editor/ContentBrowser/ContentBrowserAssetDataSource/Source/ContentBrowserAssetDataSource/Private/ContentBrowserAssetDataSource.cpp
+++ b/Engine/Plugins/Editor/ContentBrowser/ContentBrowserAssetDataSource/Source/ContentBrowserAssetDataSource/Private/ContentBrowserAssetDataSource.cpp
@@ -29,12 +29,15 @@
 #include "NewAssetContextMenu.h"
 #include "AssetFolderContextMenu.h"
 #include "AssetFileContextMenu.h"
+#include "ContentBrowserAssetDataSourceCommands.h"
+#include "ContentBrowserAssetDataSourceUtils.h"
 #include "ContentBrowserDataSubsystem.h"
 #include "ContentBrowserDataUtils.h"
 #include "ContentBrowserModule.h"
 #include "IContentBrowserSingleton.h"
 #include "ContentBrowserMenuContexts.h"
 #include "EditorDirectories.h"
+#include "Framework/Notifications/NotificationManager.h"
 #include "Misc/Paths.h"
 #include "Widgets/Input/SButton.h"
 #include "SActionButton.h"
@@ -44,6 +47,7 @@
 #include "ToolMenu.h"
 #include "ToolMenuEntry.h"
 #include "ToolMenuSection.h"
+#include "Widgets/Notifications/SNotificationList.h"
 #include "Widgets/SBoxPanel.h"
 #include "Widgets/Text/STextBlock.h"
 
@@ -162,6 +166,10 @@ void UContentBrowserAssetDataSource::Initialize(const bool InAutoRegister)
 	AssetTools = &FModuleManager::LoadModuleChecked<FAssetToolsModule>("AssetTools").Get();
 	ContentBrowserModule = &FModuleManager::LoadModuleChecked<FContentBrowserModule>("ContentBrowser");
 
+	// Extend content browser commands
+	TArray<FContentBrowserCommandExtender>& CommandExtenders = ContentBrowserModule->GetAllContentBrowserCommandExtenders();
+	CommandExtenders.Add(FContentBrowserCommandExtender::CreateUObject(this, &UContentBrowserAssetDataSource::OnExtendContentBrowserCommands));
+
 	CollectionManager = &FCollectionManagerModule::GetModule().Get();
 
 	// Listen for asset registry updates
@@ -329,6 +337,12 @@ void UContentBrowserAssetDataSource::Shutdown()
 			AssetRegistryMaybe->OnFilesLoaded().RemoveAll(this);
 		}
 	}
+	
+	if (FContentBrowserModule* ContentBrowserModulePtr = FModuleManager::GetModulePtr<FContentBrowserModule>(TEXT("ContentBrowser")))
+	{
+		TArray<FContentBrowserCommandExtender>& CommandExtenders = ContentBrowserModulePtr->GetAllContentBrowserCommandExtenders();
+		CommandExtenders.RemoveAll([this](const FContentBrowserCommandExtender& Delegate) { return Delegate.GetUObject() == this; });
+	}
 
 	FCoreUObjectDelegates::OnObjectPropertyChanged.RemoveAll(this);
 	FCoreUObjectDelegates::OnObjectPreSave.RemoveAll(this);
@@ -4227,6 +4241,44 @@ void UContentBrowserAssetDataSource::RemoveRootContentPathFromStateMachine(const
 	CurrentNode->bIsEndOfAMountPoint = false;
 }
 
+void UContentBrowserAssetDataSource::OnExtendContentBrowserCommands(TSharedRef<FUICommandList> InCommandList, FOnContentBrowserGetSelection InGetSelectionDelegate)
+{
+	const FContentBrowserAssetDataSourceCommands& AssetDataSourceCommands = FContentBrowserAssetDataSourceCommands::Get();
+
+	using namespace UE::Editor::ContentBrowserAssetDataSource;
+
+	InCommandList->MapAction(
+		AssetDataSourceCommands.CaptureThumbnail,
+		FExecuteAction::CreateWeakLambda(this, [InGetSelectionDelegate]()
+			{
+				TArray<FAssetData> SelectedAssets;
+				TArray<FString> SelectedFolders;
+				InGetSelectionDelegate.ExecuteIfBound(SelectedAssets, SelectedFolders);
+				const int32 Result = CaptureThumbnail(SelectedAssets);
+
+				// Notify user
+				FSlateNotificationManager& NotificationManager = FSlateNotificationManager::Get();
+				FNotificationInfo Info(
+					FText::FormatOrdered(
+						LOCTEXT("CaptureThumbnailNotification", "Capturing thumbnail(s) for {0} asset(s)"), 
+						FText::AsNumber(Result)
+					)
+				);
+				Info.ExpireDuration = 4.f;
+				NotificationManager.AddNotification(Info);
+			}
+		),
+		FCanExecuteAction::CreateWeakLambda(this, [InGetSelectionDelegate]()->bool
+			{
+				TArray<FAssetData> SelectedAssets;
+				TArray<FString> SelectedFolders;
+				InGetSelectionDelegate.ExecuteIfBound(SelectedAssets, SelectedFolders);
+				return CanCaptureThumbnail(SelectedAssets);
+			}
+		)
+	);
+}
+
 bool UContentBrowserAssetDataSource::PathPassesCompiledDataFilter(const FContentBrowserCompiledAssetDataFilter& InFilter, const FName InInternalPath)
 {
 	if (InFilter.ExcludedPackagePaths.Contains(InInternalPath)) // PassesExcludedPathsFilter
diff --git a/Engine/Plugins/Editor/ContentBrowser/ContentBrowserAssetDataSource/Source/ContentBrowserAssetDataSource/Private/ContentBrowserAssetDataSourceCommands.cpp b/Engine/Plugins/Editor/ContentBrowser/ContentBrowserAssetDataSource/Source/ContentBrowserAssetDataSource/Private/ContentBrowserAssetDataSourceCommands.cpp
new file mode 100644
--- /dev/null
+++ b/Engine/Plugins/Editor/ContentBrowser/ContentBrowserAssetDataSource/Source/ContentBrowserAssetDataSource/Private/ContentBrowserAssetDataSourceCommands.cpp
@@ -0,0 +1,27 @@
+// Copyright Epic Games, Inc. All Rights Reserved.
+
+#include "ContentBrowserAssetDataSourceCommands.h"
+
+#define LOCTEXT_NAMESPACE "ContentBrowserAssetDataSourceCommands"
+
+FContentBrowserAssetDataSourceCommands::FContentBrowserAssetDataSourceCommands()
+	: TCommands<FContentBrowserAssetDataSourceCommands>(
+		TEXT("ContentBrowserAssetDataSource"),
+		LOCTEXT("ContentBrowserAssetDataSource", "Content Browser Asset Data Source"),
+		NAME_None,
+		FAppStyle::GetAppStyleSetName())
+{
+}
+
+void FContentBrowserAssetDataSourceCommands::RegisterCommands()
+{
+	UI_COMMAND(
+		CaptureThumbnail,
+		"Capture Thumbnail",
+		"Captures a custom thumbnail from the active viewport for the selected asset(s)",
+		EUserInterfaceActionType::Button,
+		FInputChord(EKeys::T, EModifierKey::Control | EModifierKey::Shift)
+	);
+}
+
+#undef LOCTEXT_NAMESPACE
diff --git a/Engine/Plugins/Editor/ContentBrowser/ContentBrowserAssetDataSource/Source/ContentBrowserAssetDataSource/Private/ContentBrowserAssetDataSourceCommands.h b/Engine/Plugins/Editor/ContentBrowser/ContentBrowserAssetDataSource/Source/ContentBrowserAssetDataSource/Private/ContentBrowserAssetDataSourceCommands.h
new file mode 100644
--- /dev/null
+++ b/Engine/Plugins/Editor/ContentBrowser/ContentBrowserAssetDataSource/Source/ContentBrowserAssetDataSource/Private/ContentBrowserAssetDataSourceCommands.h
@@ -0,0 +1,18 @@
+// Copyright Epic Games, Inc. All Rights Reserved.
+
+#pragma once
+
+#include "Framework/Commands/Commands.h"
+#include "Styling/AppStyle.h"
+
+class FContentBrowserAssetDataSourceCommands : public TCommands<FContentBrowserAssetDataSourceCommands>
+{
+public:
+	FContentBrowserAssetDataSourceCommands();
+
+	//~ Begin TCommands
+	virtual void RegisterCommands() override;
+	//~ End TCommands
+
+	TSharedPtr<FUICommandInfo> CaptureThumbnail;
+};
diff --git a/Engine/Plugins/Editor/ContentBrowser/ContentBrowserAssetDataSource/Source/ContentBrowserAssetDataSource/Private/ContentBrowserAssetDataSourceModule.cpp b/Engine/Plugins/Editor/ContentBrowser/ContentBrowserAssetDataSource/Source/ContentBrowserAssetDataSource/Private/ContentBrowserAssetDataSourceModule.cpp
--- a/Engine/Plugins/Editor/ContentBrowser/ContentBrowserAssetDataSource/Source/ContentBrowserAssetDataSource/Private/ContentBrowserAssetDataSourceModule.cpp
+++ b/Engine/Plugins/Editor/ContentBrowser/ContentBrowserAssetDataSource/Source/ContentBrowserAssetDataSource/Private/ContentBrowserAssetDataSourceModule.cpp
@@ -3,6 +3,7 @@
 #include "Modules/ModuleManager.h"
 #include "UObject/StrongObjectPtr.h"
 #include "ContentBrowserAssetDataSource.h"
+#include "ContentBrowserAssetDataSourceCommands.h"
 
 class FContentBrowserAssetDataSourceModule : public FDefaultModuleImpl
 {
@@ -13,12 +14,18 @@ public:
 		{
 			AssetDataSource.Reset(NewObject<UContentBrowserAssetDataSource>(GetTransientPackage(), "AssetData"));
 			AssetDataSource->Initialize();
+			
+			FContentBrowserAssetDataSourceCommands::Register();
 		}
 	}
 
 	virtual void ShutdownModule() override
 	{
-		AssetDataSource.Reset();
+		if (GIsEditor && !IsRunningCommandlet())
+		{
+			AssetDataSource.Reset();
+			FContentBrowserAssetDataSourceCommands::Unregister();
+		}
 	}
 
 private:
diff --git a/Engine/Plugins/Editor/ContentBrowser/ContentBrowserAssetDataSource/Source/ContentBrowserAssetDataSource/Private/ContentBrowserAssetDataSourceUtils.cpp b/Engine/Plugins/Editor/ContentBrowser/ContentBrowserAssetDataSource/Source/ContentBrowserAssetDataSource/Private/ContentBrowserAssetDataSourceUtils.cpp
new file mode 100644
--- /dev/null
+++ b/Engine/Plugins/Editor/ContentBrowser/ContentBrowserAssetDataSource/Source/ContentBrowserAssetDataSource/Private/ContentBrowserAssetDataSourceUtils.cpp
@@ -0,0 +1,69 @@
+// Copyright Epic Games, Inc. All Rights Reserved.
+
+#include "ContentBrowserAssetDataSourceUtils.h"
+
+#include "AssetToolsModule.h"
+#include "AssetViewUtils.h"
+#include "Components/Viewport.h"
+#include "Editor.h"
+#include "Modules/ModuleManager.h"
+
+int32 UE::Editor::ContentBrowserAssetDataSource::CaptureThumbnail(const TArray<FAssetData>& InAssets)
+{
+	if (InAssets.IsEmpty())
+	{
+		return 0;
+	}
+
+	if (!GCurrentLevelEditingViewportClient)
+	{
+		return 0;
+	}
+
+	FViewport* Viewport = GEditor->GetActiveViewport();
+	if (!Viewport)
+	{
+		return 0;
+	}
+
+	{
+		// Rerender the requested viewport
+		FLevelEditorViewportClient* CachedViewportClient = GCurrentLevelEditingViewportClient;
+
+		// Remove selection box around client during render
+		GCurrentLevelEditingViewportClient = nullptr;
+		Viewport->Draw();
+
+		AssetViewUtils::CaptureThumbnailFromViewport(Viewport, InAssets);
+
+		// Redraw viewport to have the yellow highlight again
+		if (!GCurrentLevelEditingViewportClient)
+		{
+			GCurrentLevelEditingViewportClient = CachedViewportClient;
+		}
+
+		Viewport->Draw();
+	}
+
+	return InAssets.Num();
+}
+
+bool UE::Editor::ContentBrowserAssetDataSource::CanCaptureThumbnail(const TArray<FAssetData>& InAssets)
+{
+	if (GCurrentLevelEditingViewportClient == nullptr)
+	{
+		return false;
+	}
+
+	FAssetToolsModule* AssetToolsModule = FModuleManager::LoadModulePtr<FAssetToolsModule>("AssetTools");
+	if (!AssetToolsModule)
+	{
+		return false;
+	}
+
+	return InAssets.ContainsByPredicate([AssetToolsModule](const FAssetData& InAsset)
+		{
+			return AssetToolsModule->Get().AssetUsesGenericThumbnail(InAsset);
+		}
+	);
+}
diff --git a/Engine/Plugins/Editor/ContentBrowser/ContentBrowserAssetDataSource/Source/ContentBrowserAssetDataSource/Public/AssetFileContextMenu.h b/Engine/Plugins/Editor/ContentBrowser/ContentBrowserAssetDataSource/Source/ContentBrowserAssetDataSource/Public/AssetFileContextMenu.h
--- a/Engine/Plugins/Editor/ContentBrowser/ContentBrowserAssetDataSource/Source/ContentBrowserAssetDataSource/Public/AssetFileContextMenu.h
+++ b/Engine/Plugins/Editor/ContentBrowser/ContentBrowserAssetDataSource/Source/ContentBrowserAssetDataSource/Public/AssetFileContextMenu.h
@@ -150,6 +150,7 @@ private:
 	UE_API void ExecuteConsolidate();
 
 	/** Handler for Capture Thumbnail */
+	UE_DEPRECATED(5.8, TEXT("Use UE::Editor::ContentBrowserAssetDataSource instead"))
 	UE_API void ExecuteCaptureThumbnail();
 
 	/** Handler for Clear Thumbnail */
@@ -214,6 +215,7 @@ private:
 	UE_API bool CanExecutePropertyMatrix() const;
 
 	/** Handler to check to see if "Capture Thumbnail" can be executed */
+	UE_DEPRECATED(5.8, TEXT("Use UE::Editor::ContentBrowserAssetDataSource instead"))
 	UE_API bool CanExecuteCaptureThumbnail() const;
 
 	/** Handler to check to see if "Clear Thumbnail" should be visible */
diff --git a/Engine/Plugins/Editor/ContentBrowser/ContentBrowserAssetDataSource/Source/ContentBrowserAssetDataSource/Public/ContentBrowserAssetDataSource.h b/Engine/Plugins/Editor/ContentBrowser/ContentBrowserAssetDataSource/Source/ContentBrowserAssetDataSource/Public/ContentBrowserAssetDataSource.h
--- a/Engine/Plugins/Editor/ContentBrowser/ContentBrowserAssetDataSource/Source/ContentBrowserAssetDataSource/Public/ContentBrowserAssetDataSource.h
+++ b/Engine/Plugins/Editor/ContentBrowser/ContentBrowserAssetDataSource/Source/ContentBrowserAssetDataSource/Public/ContentBrowserAssetDataSource.h
@@ -9,6 +9,7 @@
 #include "ContentBrowserDataFilter.h"
 #include "Misc/NamePermissionList.h"
 #include "ContentBrowserDataMenuContexts.h"
+#include "ContentBrowserDelegates.h"
 #include "ContentBrowserAssetDataSource.generated.h"
 
 #define UE_API CONTENTBROWSERASSETDATASOURCE_API
@@ -461,6 +462,8 @@ private:
 	UE_API void AddRootContentPathToStateMachine(const FString& InAssetPath);
 
 	UE_API void RemoveRootContentPathFromStateMachine(const FString& InAssetPath);
+	
+	UE_API void OnExtendContentBrowserCommands(TSharedRef<FUICommandList> InCommandList, FOnContentBrowserGetSelection InGetSelectionDelegate);
 
 	IAssetRegistry* AssetRegistry = nullptr;
 
diff --git a/Engine/Plugins/Editor/ContentBrowser/ContentBrowserAssetDataSource/Source/ContentBrowserAssetDataSource/Public/ContentBrowserAssetDataSourceUtils.h b/Engine/Plugins/Editor/ContentBrowser/ContentBrowserAssetDataSource/Source/ContentBrowserAssetDataSource/Public/ContentBrowserAssetDataSourceUtils.h
new file mode 100644
--- /dev/null
+++ b/Engine/Plugins/Editor/ContentBrowser/ContentBrowserAssetDataSource/Source/ContentBrowserAssetDataSource/Public/ContentBrowserAssetDataSourceUtils.h
@@ -0,0 +1,16 @@
+// Copyright Epic Games, Inc. All Rights Reserved.
+
+#pragma once
+
+#include "Containers/Array.h"
+#include "AssetRegistry/AssetData.h"
+
+#define UE_API CONTENTBROWSERASSETDATASOURCE_API
+
+namespace UE::Editor::ContentBrowserAssetDataSource
+{
+	UE_API int32 CaptureThumbnail(const TArray<FAssetData>& InAssets);
+	UE_API bool CanCaptureThumbnail(const TArray<FAssetData>& InAssets);
+}
+
+#undef UE_API
diff --git a/Engine/Source/Editor/ContentBrowser/Private/SContentBrowser.cpp b/Engine/Source/Editor/ContentBrowser/Private/SContentBrowser.cpp
--- a/Engine/Source/Editor/ContentBrowser/Private/SContentBrowser.cpp
+++ b/Engine/Source/Editor/ContentBrowser/Private/SContentBrowser.cpp
@@ -496,7 +496,7 @@ void SContentBrowser::BindCommands()
 
 	// Allow extenders to add commands
 	FContentBrowserModule& ContentBrowserModule = FModuleManager::GetModuleChecked<FContentBrowserModule>(TEXT("ContentBrowser"));
-	TArray<FContentBrowserCommandExtender> CommandExtenderDelegates = ContentBrowserModule.GetAllContentBrowserCommandExtenders();
+	TArray<FContentBrowserCommandExtender>& CommandExtenderDelegates = ContentBrowserModule.GetAllContentBrowserCommandExtenders();
 
 	for (int32 i = 0; i < CommandExtenderDelegates.Num(); ++i)
 	{
@@ -4331,6 +4331,12 @@ void SContentBrowser::GetSelectionState(TArray<FAssetData>& SelectedAssets, TArr
 	{
 		SelectedPaths = LegacyContentSourceWidgets->PathViewPtr->GetSelectedPaths();
 	}
+	// Fallback
+	else if (const FContentBrowserModule* ContentBrowserModule = FModuleManager::GetModulePtr<FContentBrowserModule>("ContentBrowser"))
+	{
+		ContentBrowserModule->Get().GetSelectedAssets(SelectedAssets);
+		ContentBrowserModule->Get().GetSelectedFolders(SelectedPaths);
+	}
 }
 
 bool SContentBrowser::IsBackEnabled() const
