UE_DEPRECATED text should not be wrapped in TEXT()

#Synced-CL 52214067

diff --git a/Engine/Plugins/Editor/ContentBrowser/ContentBrowserAssetDataSource/Source/ContentBrowserAssetDataSource/Public/AssetFileContextMenu.h b/Engine/Plugins/Editor/ContentBrowser/ContentBrowserAssetDataSource/Source/ContentBrowserAssetDataSource/Public/AssetFileContextMenu.h
--- a/Engine/Plugins/Editor/ContentBrowser/ContentBrowserAssetDataSource/Source/ContentBrowserAssetDataSource/Public/AssetFileContextMenu.h
+++ b/Engine/Plugins/Editor/ContentBrowser/ContentBrowserAssetDataSource/Source/ContentBrowserAssetDataSource/Public/AssetFileContextMenu.h
@@ -150,7 +150,7 @@ private:
 	UE_API void ExecuteConsolidate();
 
 	/** Handler for Capture Thumbnail */
-	UE_DEPRECATED(5.8, TEXT("Use UE::Editor::ContentBrowserAssetDataSource instead"))
+	UE_DEPRECATED(5.8, "Use UE::Editor::ContentBrowserAssetDataSource instead")
 	UE_API void ExecuteCaptureThumbnail();
 
 	/** Handler for Clear Thumbnail */
@@ -215,7 +215,7 @@ private:
 	UE_API bool CanExecutePropertyMatrix() const;
 
 	/** Handler to check to see if "Capture Thumbnail" can be executed */
-	UE_DEPRECATED(5.8, TEXT("Use UE::Editor::ContentBrowserAssetDataSource instead"))
+	UE_DEPRECATED(5.8, "Use UE::Editor::ContentBrowserAssetDataSource instead")
 	UE_API bool CanExecuteCaptureThumbnail() const;
 
 	/** Handler to check to see if "Clear Thumbnail" should be visible */
