[ContentBrowser] Logic changes for the new Add Menu

Refactor the Content Browser Add Menu to support both sections and submenus within asset category paths.
Add a new FCategoryPath struct and ECategoryMenuType enum to allow asset categories to declare whether their subcategories should render as inline sections or nested submenus.
Update asset registration logic to allow statically-registerable asset definitions to replace dynamically-registered proxies, and add a new Data asset category path.

#Synced-CL 52417896

diff --git a/Engine/Plugins/Editor/ContentBrowser/ContentBrowserAssetDataSource/Source/ContentBrowserAssetDataSource/Private/NewAssetContextMenu.cpp b/Engine/Plugins/Editor/ContentBrowser/ContentBrowserAssetDataSource/Source/ContentBrowserAssetDataSource/Private/NewAssetContextMenu.cpp
--- a/Engine/Plugins/Editor/ContentBrowser/ContentBrowserAssetDataSource/Source/ContentBrowserAssetDataSource/Private/NewAssetContextMenu.cpp
+++ b/Engine/Plugins/Editor/ContentBrowser/ContentBrowserAssetDataSource/Source/ContentBrowserAssetDataSource/Private/NewAssetContextMenu.cpp
@@ -17,6 +17,7 @@
 #include "Widgets/SToolTip.h"
 #include "IDocumentation.h"
 #include "ContentBrowserUtils.h"
+#include "Internationalization/TextPackageNamespaceUtil.h"
 #include "Styling/StyleColors.h"
 #include "Widgets/SAssetMenuIcon.h"
 
@@ -42,14 +43,50 @@ struct FWizardItem
 	FSimpleDelegate OnClicked;
 };
 
+struct FCategoryChildren
+{
+	FCategoryChildren(FName InSubCategoryName, TSharedPtr<FCategorySubMenuItem> InSubMenuItem)
+		: SubCategoryName(InSubCategoryName)
+		, SubMenuItem(InSubMenuItem)
+	{}
+
+	FName SubCategoryName;
+	TSharedPtr<FCategorySubMenuItem> SubMenuItem;
+};
+
+struct FFactoryInfo
+{
+	explicit FFactoryInfo(const FFactoryItem& InFactory)
+		: Factory(InFactory)
+		, SubSectionName(FText::GetEmpty())
+	{}
+
+	explicit FFactoryInfo(const FFactoryItem& InFactory, const FText& InSubSectionName)
+		: Factory(InFactory)
+		, SubSectionName(InSubSectionName)
+	{}
+
+	FFactoryItem Factory;
+	FText SubSectionName;
+};
+
 struct FCategorySubMenuItem
 {
 	FText Name;
-	TArray<FFactoryItem> Factories;
+	TArray<FFactoryInfo> Factories;
 	TArray<FWizardItem> Wizards;
-	TMap<FString, TSharedPtr<FCategorySubMenuItem>> Children;
+	ECategoryMenuType CategoryMenuType = ECategoryMenuType::Menu;
+	TArray<FCategoryChildren> Children;
 
-	void SortSubMenus(FCategorySubMenuItem* SubMenu = nullptr)
+	void GenerateChildrenSubMenuItemArray(TArray<TSharedPtr<FCategorySubMenuItem>>& OutSubMenuItems)
+	{
+		for (const FCategoryChildren& Child : Children)
+		{
+			OutSubMenuItems.Add(Child.SubMenuItem);
+		}
+	}
+
+	void SortSubMenus(const FText& InMainCategory,FCategorySubMenuItem* SubMenu = nullptr)
 	{
 		if (!SubMenu)
 		{
@@ -57,17 +94,32 @@ struct FCategorySubMenuItem
 		}
 
 		// Sort the factories by display name
-		SubMenu->Factories.Sort([](const FFactoryItem& A, const FFactoryItem& B) -> bool
+		SubMenu->Factories.Sort([InMainCategory](const FFactoryInfo& A, const FFactoryInfo& B) -> bool
 		{
-			return A.DisplayName.CompareToCaseIgnored(B.DisplayName) < 0;
+			if (A.SubSectionName.CompareTo(B.SubSectionName) == 0)
+			{
+				return A.Factory.DisplayName.CompareToCaseIgnored(B.Factory.DisplayName) < 0;
+			}
+
+			if (B.SubSectionName.CompareToCaseIgnored(InMainCategory) == 0)
+			{
+				return false;
+			}
+
+			if (A.SubSectionName.CompareToCaseIgnored(InMainCategory) == 0)
+			{
+				return true;
+			}
+
+			return A.SubSectionName.CompareTo(B.SubSectionName) < 0;
 		});
 
-		for (TPair<FString, TSharedPtr<FCategorySubMenuItem>>& Pair : SubMenu->Children)
+		for (FCategoryChildren& Child : SubMenu->Children)
 		{
-			if (Pair.Value.IsValid())
+			if (Child.SubMenuItem.IsValid())
 			{
-				FCategorySubMenuItem* MenuData = Pair.Value.Get();
-				SortSubMenus(MenuData);
+				FCategorySubMenuItem* MenuData = Child.SubMenuItem.Get();
+				SortSubMenus(InMainCategory, MenuData);
 			}
 		}
 	}
@@ -376,30 +428,119 @@ void FNewAssetContextMenu::CreateNewAssetMenuCategory(UToolMenu* Menu, FName Sec
 		return;
 	}
 
+	FAssetToolsModule& AssetToolsModule = FModuleManager::LoadModuleChecked<FAssetToolsModule>("AssetTools");
+	TArray<FAdvancedAssetCategory> CategoryNames;
+	AssetToolsModule.Get().GetAllAdvancedAssetCategories(CategoryNames);
+	FAdvancedAssetCategory* Category = CategoryNames.FindByPredicate([AssetTypeCategory](const FAdvancedAssetCategory& InAdvancedCategory)
+		{
+			return InAdvancedCategory.CategoryType == AssetTypeCategory;
+		});
+	// Misc will be used as the default category if the category can't be found
+	FName CategoryName = Category ? FName(Category->CategoryName.BuildSourceString()) : FName(TEXT("Misc"));
+
 	TSharedPtr<FCategorySubMenuItem> ParentMenuData = MakeShared<FCategorySubMenuItem>();
 	for (FFactoryItem& Item : FactoriesInThisCategory)
 	{
 		FCategorySubMenuItem* SubMenu = ParentMenuData.Get();
-		const TArray<FText>& CategoryNames = Item.Factory.GetMenuCategorySubMenus();
-		for (FText CategoryName : CategoryNames)
+		// If we need to populate the basic section we can skip everything else as there is not a category for it to iterate.
+		if (AssetTypeCategory == EAssetTypeCategories::Basic)
+		{
+			SubMenu->Factories.Add(FFactoryInfo(Item));
+			continue;
+		}
+
+		TArray<FAssetCategoryPath> CategoryPaths = Item.Factory.GetAssetMenuPathsForCategory(CategoryName);
+
+		// Always add a default category path based on the major one so that SubCategory is set
+		if (CategoryPaths.IsEmpty())
 		{
-			const FString SourceString = CategoryName.BuildSourceString();
-			if (TSharedPtr<FCategorySubMenuItem> SubMenuData = SubMenu->Children.FindRef(SourceString))
+			// Try to add the sub categories from the sub menus array
+			const TArray<FText> SubCategories = Item.Factory.GetMenuCategorySubMenus();
+			if (!SubCategories.IsEmpty())
 			{
-				check(SubMenuData.IsValid());
-				SubMenu = SubMenuData.Get();
+				for (const FText& SubCategory : SubCategories)
+				{
+					CategoryPaths.Add(FAssetCategoryPath(FText::FromName(CategoryName), SubCategory));
+				}
 			}
 			else
 			{
-				TSharedPtr<FCategorySubMenuItem> NewSubMenu = MakeShared<FCategorySubMenuItem>();
-				NewSubMenu->Name = CategoryName;
-				SubMenu->Children.Add(SourceString, NewSubMenu);
-				SubMenu = NewSubMenu.Get();
+				CategoryPaths.Add(FAssetCategoryPath(FText::FromName(CategoryName)));
+			}
+		}
+
+		// Final Section name to use in the menu for the entry
+		TArray<FText> SubCategoriesSection;
+
+		for (const FAssetCategoryPath& CategoryPath : CategoryPaths)
+		{
+			TArray<FCategoryPath> OutSubCategories;
+			CategoryPath.GetSubCategoriesInfo(OutSubCategories);
+			// For new full category path reset the category section and SubMenu
+			SubCategoriesSection.Empty();
+			SubMenu = ParentMenuData.Get();
+
+			// All basic section will have the basic category
+			if (OutSubCategories.IsEmpty())
+			{
+				OutSubCategories.Add(FCategoryPath(LOCTEXT("BasicSectionName", "Basic"), ECategoryMenuType::Section));
+			}
+
+			for (FCategoryPath SubCategory : OutSubCategories)
+			{
+				// If this is a sub-sub... menu remove the old section entries, this is currently only possible with the old AssetTypeAction.
+				if (SubMenu != ParentMenuData.Get() && SubMenu->CategoryMenuType == ECategoryMenuType::Menu)
+				{
+					SubCategoriesSection.Empty();
+				}
+
+				const FString SourceString = SubCategory.GetSubMenuName().BuildSourceString();
+				const FName SourceName = FName(SourceString);
+				const FText SubCategorySectionName = SubCategory.GetSubMenuName();
+				if (!SubCategorySectionName.IsEmpty())
+				{
+					SubCategoriesSection.Add(SubCategorySectionName);
+				}
+
+				FCategoryChildren* SubMenuData = SubMenu->Children.FindByPredicate(
+					[SourceName, CategoryType = SubCategory.GetCategoryMenuType()] 
+						(const FCategoryChildren& InCategoryChild)
+					{
+						// Get the entry that match the name and type
+						return InCategoryChild.SubCategoryName == SourceName && 
+								InCategoryChild.SubMenuItem.IsValid() && 
+								InCategoryChild.SubMenuItem->CategoryMenuType == CategoryType;
+					});
+
+				if (SubMenuData)
+				{
+					SubMenu = SubMenuData->SubMenuItem.Get();
+				}
+				else
+				{
+					TSharedPtr<FCategorySubMenuItem> NewSubMenu = MakeShared<FCategorySubMenuItem>();
+					NewSubMenu->Name = SubCategory.GetSubMenuName();
+					NewSubMenu->CategoryMenuType = SubCategory.GetCategoryMenuType();
+					SubMenu->Children.Add(FCategoryChildren(SourceName, NewSubMenu));
+					SubMenu = NewSubMenu.Get();
+				}
+			}
+			
+			// Assign the factory after a category path is read so that if the asset is inside 2 different section/menu in the same category it will be assigned correctly
+			if (SubCategoriesSection.IsEmpty())
+			{
+				SubMenu->Factories.Add(FFactoryInfo(Item, FText::FromName(CategoryName)));
+			}
+			else
+			{
+				for (const FText& SubCategorySection : SubCategoriesSection)
+				{
+					SubMenu->Factories.Add(FFactoryInfo(Item, SubCategorySection));
+				}
 			}
 		}
-		SubMenu->Factories.Add(Item);
 	}
-	ParentMenuData->SortSubMenus();
+	ParentMenuData->SortSubMenus(FText::FromName(CategoryName));
 
 	// Find wizards
 	FContentBrowserModule& ContentBrowserModule = FModuleManager::GetModuleChecked<FContentBrowserModule>(TEXT("ContentBrowser"));
@@ -410,19 +551,28 @@ void FNewAssetContextMenu::CreateNewAssetMenuCategory(UToolMenu* Menu, FName Sec
 			FCategorySubMenuItem* SubMenu = ParentMenuData.Get();
 
 			const FString SourceString = Wizard.CategoryPath.CategoryName.BuildSourceString();
-			if (TSharedPtr<FCategorySubMenuItem> SubMenuData = SubMenu->Children.FindRef(SourceString))
+			const FName SourceName = FName(SourceString);
+			FCategoryChildren* SubMenuData = SubMenu->Children.FindByPredicate([SourceName] (const FCategoryChildren& InCategoryChild)
+				{
+					// Get the entry that match the name and type
+					return InCategoryChild.SubCategoryName == SourceName &&
+							InCategoryChild.SubMenuItem.IsValid() && 
+							// Wizard do not currently support the menu type or sub section name so consider them all in the menu category which was the default before
+							InCategoryChild.SubMenuItem->CategoryMenuType == ECategoryMenuType::Menu;
+				});
+
+			if (SubMenuData && SubMenuData->SubMenuItem.IsValid())
 			{
-				check(SubMenuData.IsValid());
-				SubMenu = SubMenuData.Get();
+				SubMenu = SubMenuData->SubMenuItem.Get();
 			}
 			else
 			{
 				TSharedPtr<FCategorySubMenuItem> NewSubMenu = MakeShared<FCategorySubMenuItem>();
 				NewSubMenu->Name = Wizard.CategoryPath.CategoryName;
-				SubMenu->Children.Add(SourceString, NewSubMenu);
+				SubMenu->Children.Add(FCategoryChildren(SourceName, NewSubMenu));
 				SubMenu = NewSubMenu.Get();
 			}
-			
+
 			FWizardItem& WizardItem = SubMenu->Wizards.AddDefaulted_GetRef();
 			WizardItem.DisplayName = Wizard.DisplayName;
 			WizardItem.Description = Wizard.Description;
@@ -457,9 +607,21 @@ void FNewAssetContextMenu::CreateNewAssetMenus(UToolMenu* Menu, FName SectionNam
 	{
 		Section.AddSeparator(NAME_None);
 	}
-	
-	for (const FFactoryItem& FactoryItem : SubMenuData->Factories)
+
+	for (const FFactoryInfo& FactoryInfo : SubMenuData->Factories)
 	{
+		FToolMenuSection* SectionToUse = &Section;
+
+		if (!FactoryInfo.SubSectionName.IsEmpty())
+		{
+			const FName EntrySectionName = FName(FactoryInfo.SubSectionName.BuildSourceString());
+			SectionToUse = &Menu->FindOrAddSection(EntrySectionName, FactoryInfo.SubSectionName);
+			// Add the label anyway since it may have been created without it before
+			SectionToUse->Label = FactoryInfo.SubSectionName;
+		}
+
+		const FFactoryItem& FactoryItem = FactoryInfo.Factory;
+
 		TWeakObjectPtr<UClass> WeakFactoryClass = FactoryItem.Factory.GetClass();
 
 		FName AssetTypeName;
@@ -474,7 +636,7 @@ void FNewAssetContextMenu::CreateNewAssetMenus(UToolMenu* Menu, FName SectionNam
 			static constexpr uint32 IconContainerSize = 24;
 			static constexpr uint32 IconSize = 16;
 
-			FToolMenuEntry& Entry = Section.AddEntry(
+			FToolMenuEntry& Entry = SectionToUse->AddEntry(
 				FToolMenuEntry::InitMenuEntry(
 				NAME_None,
 				FUIAction(
@@ -491,7 +653,7 @@ void FNewAssetContextMenu::CreateNewAssetMenus(UToolMenu* Menu, FName SectionNam
 		}
 		else
 		{
-			Section.AddEntry(
+			SectionToUse->AddEntry(
 				FToolMenuEntry::InitMenuEntry(
 				NAME_None,
 				FUIAction(
@@ -509,41 +671,82 @@ void FNewAssetContextMenu::CreateNewAssetMenus(UToolMenu* Menu, FName SectionNam
 		return;
 	}
 
-	Section.AddSeparator(NAME_None);
-
 	TArray<TSharedPtr<FCategorySubMenuItem>> SortedMenus;
-	SubMenuData->Children.GenerateValueArray(SortedMenus);
-	SortedMenus.Sort([](const TSharedPtr<FCategorySubMenuItem>& A, const TSharedPtr<FCategorySubMenuItem>& B) -> bool
+	SubMenuData->GenerateChildrenSubMenuItemArray(SortedMenus);
+	FText BasicSection = LOCTEXT("BasicSection", "Basic");
+	SortedMenus.Sort([BasicSection](const TSharedPtr<FCategorySubMenuItem>& A, const TSharedPtr<FCategorySubMenuItem>& B) -> bool
 	{
+		// Basic section should always be on top
+		if (B->Name.CompareToCaseIgnored(BasicSection) == 0)
+		{
+			return false;
+		}
+
+		// Basic section should always be on top
+		if (A->Name.CompareToCaseIgnored(BasicSection) == 0)
+		{
+			return true;
+		}
+
+		// SubMenus should go after sections
+		if (A->CategoryMenuType != B->CategoryMenuType)
+		{
+			return A->CategoryMenuType == ECategoryMenuType::Section;
+		}
+
 		return A->Name.CompareToCaseIgnored(B->Name) < 0;
 	});
 
+	bool bWasSeparatorAdded = false;
 	for (TSharedPtr<FCategorySubMenuItem>& ChildMenuData : SortedMenus)
 	{
 		check(ChildMenuData.IsValid());
 
-		FToolMenuEntry& Entry = Section.AddSubMenu(
-			NAME_None,
-			ChildMenuData->Name,
-			FText::GetEmpty(),
-			FNewToolMenuDelegate::CreateStatic(
-				&FNewAssetContextMenu::CreateNewAssetMenus,
-				FName("Section"),
+		if (ChildMenuData->CategoryMenuType == ECategoryMenuType::Section)
+		{
+			const FName SeparatorName = FName(ChildMenuData->Name.BuildSourceString());
+			CreateNewAssetMenus(
+				Menu,
+				SeparatorName,
 				ChildMenuData,
 				InPath,
 				InOnNewAssetRequested,
 				InCanExecuteAction
-			),
-			FUIAction(
-				FExecuteAction(),
-				InCanExecuteAction
-			),
-			EUserInterfaceActionType::Button
-		);
-
-		if (UE::Editor::ContentBrowser::IsNewStyleEnabled())
+				);
+		}
+		else
 		{
-			Entry.SubMenuData.Style.StyleName = "ContentBrowser.AddNewMenu";
+			FToolMenuSection& NewSubMenuSection = Menu->FindOrAddSection(FName(TEXT("SubMenuSection")), FText::GetEmpty(), FToolMenuInsert(NAME_None, EToolMenuInsertType::Last));
+			FName SubMenuName = FName(ChildMenuData->Name.BuildSourceString());
+			if (!bWasSeparatorAdded)
+			{
+				bWasSeparatorAdded = true;
+				NewSubMenuSection.AddSeparator(NAME_None);
+			}
+
+			FToolMenuEntry& Entry = NewSubMenuSection.AddSubMenu(
+				SubMenuName,
+				ChildMenuData->Name,
+				FText::GetEmpty(),
+				FNewToolMenuDelegate::CreateStatic(
+					&FNewAssetContextMenu::CreateNewAssetMenus,
+					SubMenuName,
+					ChildMenuData,
+					InPath,
+					InOnNewAssetRequested,
+					InCanExecuteAction
+				),
+				FUIAction(
+					FExecuteAction(),
+					InCanExecuteAction
+				),
+				EUserInterfaceActionType::Button
+			);
+
+			if (UE::Editor::ContentBrowser::IsNewStyleEnabled())
+			{
+				Entry.SubMenuData.Style.StyleName = "ContentBrowser.AddNewMenu";
+			}
 		}
 	}
 }
diff --git a/Engine/Source/Developer/AssetTools/Private/AssetTools.cpp b/Engine/Source/Developer/AssetTools/Private/AssetTools.cpp
--- a/Engine/Source/Developer/AssetTools/Private/AssetTools.cpp
+++ b/Engine/Source/Developer/AssetTools/Private/AssetTools.cpp
@@ -40,7 +40,6 @@
 #include "AssetTypeActions/AssetTypeActions_AnimBoneCompressionSettings.h"
 #include "AssetTypeActions/AssetTypeActions_AnimCurveCompressionSettings.h"
 #include "AssetTypeActions/AssetTypeActions_VariableFrameStrippingSettings.h"
-#include "AssetTypeActions/AssetTypeActions_PhysicalMaterialMask.h"
 #include "AssetTypeActions/AssetTypeActions_DeletedObjectPlaceholder.h"
 #include "WorldPartition/WorldPartition.h"
 #include "SDiscoveringAssetsDialog.h"
@@ -1422,7 +1421,6 @@ UAssetToolsImpl::UAssetToolsImpl(const FObjectInitializer& ObjectInitializer)
 	RegisterAssetTypeActions(MakeShareable(new FAssetTypeActions_AnimBoneCompressionSettings));
 	RegisterAssetTypeActions(MakeShareable(new FAssetTypeActions_AnimCurveCompressionSettings));
 	RegisterAssetTypeActions(MakeShareable(new FAssetTypeActions_VariableFrameStrippingSettings));
-	RegisterAssetTypeActions(MakeShareable(new FAssetTypeActions_PhysicalMaterialMask));
 	RegisterAssetTypeActions(MakeShareable(new FAssetTypeActions_DeletedObjectPlaceholder));
 
 	// Note: Please don't add any more actions here!  They belong in an editor-only module that is more tightly
@@ -1490,8 +1488,35 @@ void UAssetToolsImpl::SyncAssetTypesToAssetDefinitions()
 				const TSoftClassPtr<UObject> AssetClass = AssetDefinition->GetAssetClass();
 				const bool ActionDefinitionAlreadyProxied = 
 					Algo::AnyOf(AssetTypeActionsList, [AssetClass](const TSharedRef<IAssetTypeActions>& Actions){ return Actions->GetSupportedClass() == AssetClass; });
-			
-				if (!ActionDefinitionAlreadyProxied)
+
+				// If an asset definition is already registered check that the new one doesn't take precedence over the old one
+				if (ActionDefinitionAlreadyProxied)
+				{
+					if (const int* Index = AssetTypeActionsLookup.Find(AssetClass.Get()))
+					{
+						TWeakPtr<IAssetTypeActions> TypeActions = AssetTypeActionsList[*Index];
+						if (const TSharedPtr<IAssetTypeActions>& ActionPtr = TypeActions.Pin())
+						{
+							if (ActionPtr->IsAssetDefinitionInDisguise())
+							{
+								const TSharedPtr<FAssetDefinitionProxy>& ProxyAction = StaticCastSharedPtr<FAssetDefinitionProxy>(ActionPtr);
+								if (const UAssetDefinition* ProxyAssetDefinition = ProxyAction->GetAssetDefinition())
+								{
+									// If the old one can't be registered statically, but the new one can, we should remove the old in favor of the new one
+									if (!ProxyAssetDefinition->CanRegisterStatically() && AssetDefinition->CanRegisterStatically())
+									{
+										// Unregister the old one in favor of the one that can be registered statically
+										UnregisterAssetTypeActions(ProxyAction.ToSharedRef());
+										TSharedRef<FAssetDefinitionProxy> Proxy = MakeShared<FAssetDefinitionProxy>(AssetDefinition);
+										// Cache the asset definition proxy.
+										RegisterAssetTypeActions(Proxy);
+									}
+								}
+							}
+						}
+					}	
+				}
+				else
 				{
 					TSharedRef<FAssetDefinitionProxy> Proxy = MakeShared<FAssetDefinitionProxy>(AssetDefinition);
 					// Cache the asset definition proxy.
diff --git a/Engine/Source/Editor/AssetDefinition/Private/AssetDefinition.cpp b/Engine/Source/Editor/AssetDefinition/Private/AssetDefinition.cpp
--- a/Engine/Source/Editor/AssetDefinition/Private/AssetDefinition.cpp
+++ b/Engine/Source/Editor/AssetDefinition/Private/AssetDefinition.cpp
@@ -30,6 +30,7 @@ FAssetCategoryPath EAssetCategoryPaths::Foliage(LOCTEXT("Foliage", "Foliage"));
 FAssetCategoryPath EAssetCategoryPaths::Input(LOCTEXT("Input", "Input"));
 FAssetCategoryPath EAssetCategoryPaths::FX(LOCTEXT("FX", "FX"));
 FAssetCategoryPath EAssetCategoryPaths::Cinematics(LOCTEXT("Cinematics", "Cinematics"));
+FAssetCategoryPath EAssetCategoryPaths::Data(LOCTEXT("Data", "Data"));
 FAssetCategoryPath EAssetCategoryPaths::Media(LOCTEXT("Media", "Media"));
 FAssetCategoryPath EAssetCategoryPaths::World(LOCTEXT("World", "World"));
 
@@ -143,21 +144,12 @@ namespace UE::AssetDefinition::Status
 
 FAssetCategoryPath::FAssetCategoryPath(const FText& InCategory)
 {
-	CategoryPath = { TPair<FName, FText>(FName(*FTextInspector::GetSourceString(InCategory)), InCategory) };
+	CategoryPath = { TPair<FName, FCategoryPath>(FName(*FTextInspector::GetSourceString(InCategory)), FCategoryPath(InCategory)) };
 }
 
-FAssetCategoryPath::FAssetCategoryPath(const FText& InCategory, const FText& InSubCategory)
+void FAssetCategoryPath::AppendSubMenu(const FCategoryPath& InSubMenu)
 {
-	CategoryPath = {
-		TPair<FName, FText>(FName(*FTextInspector::GetSourceString(InCategory)), InCategory),
-		TPair<FName, FText>(FName(*FTextInspector::GetSourceString(InSubCategory)), InSubCategory)
-	};
-}
-
-FAssetCategoryPath::FAssetCategoryPath(const FAssetCategoryPath& InCategory, const FText& InSubCategory)
-{
-    CategoryPath.Append(InCategory.CategoryPath);
-    CategoryPath.Add(TPair<FName, FText>(FName(*FTextInspector::GetSourceString(InSubCategory)), InSubCategory));
+	CategoryPath.Add(TPair<FName, FCategoryPath>(FName(*FTextInspector::GetSourceString(InSubMenu.GetSubMenuName())), InSubMenu));
 }
 
 FAssetCategoryPath::FAssetCategoryPath(TConstArrayView<FText> InCategoryPath)
@@ -166,7 +158,7 @@ FAssetCategoryPath::FAssetCategoryPath(TConstArrayView<FText> InCategoryPath)
 	
 	for (const FText& CategoryChunk : InCategoryPath)
 	{
-		CategoryPath.Add(TPair<FName, FText>(FName(*FTextInspector::GetSourceString(CategoryChunk)), CategoryChunk));
+		CategoryPath.Add(TPair<FName, FCategoryPath>(FName(*FTextInspector::GetSourceString(CategoryChunk)), FCategoryPath(CategoryChunk)));
 	}
 }
 
@@ -183,6 +175,18 @@ void FAssetCategoryPath::GetSubCategories(TArray<FName>& SubCategories) const
 }
 
 void FAssetCategoryPath::GetSubCategoriesText(TArray<FText>& SubCategories) const
+{
+	if (HasSubCategory())
+	{
+		SubCategories.Reserve(SubCategories.Num() + NumSubCategories());
+		for (int32 i = 1; i < CategoryPath.Num(); i++)
+		{
+			SubCategories.Add(CategoryPath[i].Value.GetSubMenuName());
+		}
+	}
+}
+
+void FAssetCategoryPath::GetSubCategoriesInfo(TArray<FCategoryPath>& SubCategories) const
 {
 	if (HasSubCategory())
 	{
diff --git a/Engine/Source/Editor/AssetDefinition/Public/AssetDefinition.h b/Engine/Source/Editor/AssetDefinition/Public/AssetDefinition.h
--- a/Engine/Source/Editor/AssetDefinition/Public/AssetDefinition.h
+++ b/Engine/Source/Editor/AssetDefinition/Public/AssetDefinition.h
@@ -280,6 +280,7 @@ struct EAssetCategoryPaths
 	static UE_API FAssetCategoryPath Audio;
 	static UE_API FAssetCategoryPath Blueprint;
 	static UE_API FAssetCategoryPath Cinematics;
+	static UE_API FAssetCategoryPath Data;
 	static UE_API FAssetCategoryPath Foliage;
 	static UE_API FAssetCategoryPath FX;
 	static UE_API FAssetCategoryPath Gameplay;
@@ -644,7 +645,6 @@ public:
 	// 3. If you add a new function, please create a struct for the Args.  We'll be able to upgrade things easier.
 	//    Please continue to use EAssetCommandResult and FAssetSupportResponse, for those kinds of commands.
 
-protected:
 	UE_API virtual bool CanRegisterStatically() const;
 
 protected:
diff --git a/Engine/Source/Editor/AssetDefinition/Public/Misc/AssetCategoryPath.h b/Engine/Source/Editor/AssetDefinition/Public/Misc/AssetCategoryPath.h
--- a/Engine/Source/Editor/AssetDefinition/Public/Misc/AssetCategoryPath.h
+++ b/Engine/Source/Editor/AssetDefinition/Public/Misc/AssetCategoryPath.h
@@ -6,33 +6,89 @@
 
 #define UE_API ASSETDEFINITION_API
 
+/** 
+ * Whether the SubCategory should be a menu or just a section
+ */
+enum class ECategoryMenuType : uint8
+{
+	Menu = 0,
+	Section = 1 << 0
+};
+
+/**
+ * Contains the sub menu name and the type of the sub menu, either a menu or a section, section should be the last in the chain
+ */
+struct FCategoryPath
+{
+	FCategoryPath(const FText& InSubMenuName, const ECategoryMenuType InCategoryMenuType = ECategoryMenuType::Menu)
+		: SubMenuName(InSubMenuName)
+		, CategoryMenuType(InCategoryMenuType)
+	{}
+
+	/** Get the name of the sub menu, could be a sub category based on the position inside the FAssetCategoryPath */
+	FText GetSubMenuName() const { return SubMenuName; }
+
+	/** Either Menu or Section, based on if you want a menu or a section for your entry */
+	ECategoryMenuType GetCategoryMenuType() const { return CategoryMenuType; }
+
+private:
+	FText SubMenuName;
+	ECategoryMenuType CategoryMenuType;
+};
+
 /**
- * The asset category path is how we know how to build menus around assets.  For example, Basic is generally the ones
+ * The asset category path is how we know how to build menus around assets.\n  For example, Basic is generally the ones
  * we expose at the top level, where as everything else is a category with a pull out menu, and the subcategory would
- * be where it gets placed in a submenu inside of there.
+ * be where it gets placed in a submenu inside of there.\n\n
+ * Currently, the system supports chaining menus and as the last entry you can have a custom section.\n
+ * You can't have a menu inside a section so the section should be the last in case you want one.
  */
 struct FAssetCategoryPath
 {
 	UE_API FAssetCategoryPath(const FText& InCategory);
-	UE_API FAssetCategoryPath(const FText& InCategory, const FText& SubCategory);
-	UE_API FAssetCategoryPath(const FAssetCategoryPath& InCategory, const FText& InSubCategory);
 	UE_API FAssetCategoryPath(TConstArrayView<FText> InCategoryPath);
 
+	/** 
+	 * Constructs from a root FAssetCategoryPath followed by any number of FCategoryPath.\n
+	 * By default, every FText given is treated like a menu.\n
+	 * If the last entry is a section you should still explicitly add it like so:\n
+	 * FAssetCategoryPath(FAssetCategoryPath or (FText), FCategoryPath or (FText)..., FCategoryPath(FText, ECategoryMenuType::Section))
+	 */
+	template <typename... TInSubMenus>
+	requires (std::is_constructible_v<FCategoryPath, TInSubMenus> && ...)
+	FAssetCategoryPath(const FAssetCategoryPath& InCategory, TInSubMenus&&... InSubMenus)
+	{
+		// Append the first category
+		CategoryPath.Append(InCategory.CategoryPath);
+
+		// Append all the additional sub-menus/sections
+		(AppendSubMenu(FCategoryPath(Forward<TInSubMenus>(InSubMenus))), ...);
+	}
+
+	/** Constructor helper to avoid explicit wrapping when a specific menu type is needed */
+	FAssetCategoryPath(const FAssetCategoryPath& InCategory, const FText& InSubMenuName, ECategoryMenuType InMenuType = ECategoryMenuType::Menu)
+		: FAssetCategoryPath(InCategory, FCategoryPath(InSubMenuName, InMenuType))
+	{}
+
 	FName GetCategory() const { return CategoryPath[0].Key; }
-	FText GetCategoryText() const { return CategoryPath[0].Value; }
+	FText GetCategoryText() const { return CategoryPath[0].Value.GetSubMenuName(); }
 	
 	bool HasSubCategory() const { return CategoryPath.Num() > 1; }
 	int32 NumSubCategories() const { return HasSubCategory() ? (CategoryPath.Num() - 1) : 0; }
 	FName GetSubCategory() const { return HasSubCategory() ? CategoryPath[1].Key : NAME_None; }
-	FText GetSubCategoryText() const { return HasSubCategory() ? CategoryPath[1].Value : FText::GetEmpty(); }
+	FText GetSubCategoryText() const { return HasSubCategory() ? CategoryPath[1].Value.GetSubMenuName() : FText::GetEmpty(); }
 
 	UE_API void GetSubCategories(TArray<FName>& SubCategories) const;
 	UE_API void GetSubCategoriesText(TArray<FText>& SubCategories) const;
+	UE_API void GetSubCategoriesInfo(TArray<FCategoryPath>& SubCategories) const;
 	
 	FAssetCategoryPath operator / (const FText& SubCategory) const { return FAssetCategoryPath(*this, SubCategory); }
 	
 private:
-	TArray<TPair<FName, FText>> CategoryPath;
+	UE_API void AppendSubMenu(const FCategoryPath& InSubMenu);
+
+private:
+	TArray<TPair<FName, FCategoryPath>> CategoryPath;
 };
 
 #undef UE_API
diff --git a/Engine/Source/Editor/UnrealEd/Classes/Factories/Factory.h b/Engine/Source/Editor/UnrealEd/Classes/Factories/Factory.h
--- a/Engine/Source/Editor/UnrealEd/Classes/Factories/Factory.h
+++ b/Engine/Source/Editor/UnrealEd/Classes/Factories/Factory.h
@@ -3,6 +3,7 @@
 #pragma once
 
 #include "CoreMinimal.h"
+#include "Misc/AssetCategoryPath.h"
 #include "UObject/ObjectMacros.h"
 #include "UObject/Object.h"
 #include "Templates/SubclassOf.h"
@@ -169,6 +170,9 @@ public:
 	/** Branch of sub-menus containing factory under each provided category. */
 	UNREALED_API virtual const TArray<FText>& GetMenuCategorySubMenus() const;
 
+	/** Get the asset category path using the AssetDefinition format, used mainly for the ContentBrowser add menu. */
+	UNREALED_API virtual TArray<FAssetCategoryPath> GetAssetMenuPathsForCategory(FName InCategory) const;
+
 	/** Returns the tooltip text description of this factory */
 	UNREALED_API virtual FText GetToolTip() const;
 
@@ -410,6 +414,13 @@ protected:
 	{
 		return FactoryCreateBinary(InClass, InParent, InName, Flags, Context, Type, Buffer, BufferEnd, Warn);
 	}
+
+	/** 
+	 * Utility function that grabs the asset definition categories for a specific class and category
+	 * Used mainly for factories not having the actual class for its supported class. 
+	 */
+	UNREALED_API TArray<FAssetCategoryPath> GetAssetMenuPaths(const UClass* InClass, FName InCategory) const;
+
 public:
 
 	/** The class manufactured by this factory. */
diff --git a/Engine/Source/Editor/UnrealEd/Private/Factories/Factory.cpp b/Engine/Source/Editor/UnrealEd/Private/Factories/Factory.cpp
--- a/Engine/Source/Editor/UnrealEd/Private/Factories/Factory.cpp
+++ b/Engine/Source/Editor/UnrealEd/Private/Factories/Factory.cpp
@@ -1,6 +1,8 @@
 // Copyright Epic Games, Inc. All Rights Reserved.
 
 #include "Factories/Factory.h"
+
+#include "AssetDefinitionRegistry.h"
 #include "Misc/MessageDialog.h"
 #include "HAL/FileManager.h"
 #include "Misc/FileHelper.h"
@@ -51,6 +53,32 @@ void UFactory::AddReferencedObjects(UObject* InThis, FReferenceCollector& Collec
 	Super::AddReferencedObjects(This, Collector);
 }
 
+TArray<FAssetCategoryPath> UFactory::GetAssetMenuPaths(const UClass* InClass, FName InCategory) const
+{
+	const UClass* LocalSupportedClass = InClass;
+
+	if (UAssetDefinitionRegistry::Get() && LocalSupportedClass)
+	{
+		const UAssetDefinitionRegistry* AssetDefinitionRegistry = UAssetDefinitionRegistry::Get();
+		const UAssetDefinition* AssetDefinition = AssetDefinitionRegistry->GetAssetDefinitionForClass(LocalSupportedClass);
+		if (AssetDefinition)
+		{
+			TArray<FAssetCategoryPath> AssetCategoryPaths;
+			for (const FAssetCategoryPath& AssetCategoryPath : AssetDefinition->GetAssetCategories())
+			{
+				if (AssetCategoryPath.GetCategory() == InCategory)
+				{
+					AssetCategoryPaths.Add(AssetCategoryPath);
+				}
+			}
+			return AssetCategoryPaths;
+		}
+	}
+
+	const static TArray<FAssetCategoryPath> DefaultSubCategories;
+	return DefaultSubCategories;
+}
+
 
 UObject* UFactory::FactoryCreateFile(UClass* InClass, UObject* InParent, FName InName, EObjectFlags Flags, const FString& Filename, const TCHAR* Parms, FFeedbackContext* Warn, bool& bOutOperationCanceled)
 {
@@ -303,6 +331,11 @@ const TArray<FText>& UFactory::GetMenuCategorySubMenus() const
 	return DefaultSubCategories;
 }
 
+TArray<FAssetCategoryPath> UFactory::GetAssetMenuPathsForCategory(FName InCategory) const
+{
+	return GetAssetMenuPaths(GetSupportedClass(), InCategory);
+}
+
 FText UFactory::GetToolTip() const
 {
 	return GetSupportedClass()->GetToolTipText();
