[ContentBrowser] Add validity check before populating sub menu array

#Synced-CL 52417896

diff --git a/Engine/Plugins/Editor/ContentBrowser/ContentBrowserAssetDataSource/Source/ContentBrowserAssetDataSource/Private/NewAssetContextMenu.cpp b/Engine/Plugins/Editor/ContentBrowser/ContentBrowserAssetDataSource/Source/ContentBrowserAssetDataSource/Private/NewAssetContextMenu.cpp
--- a/Engine/Plugins/Editor/ContentBrowser/ContentBrowserAssetDataSource/Source/ContentBrowserAssetDataSource/Private/NewAssetContextMenu.cpp
+++ b/Engine/Plugins/Editor/ContentBrowser/ContentBrowserAssetDataSource/Source/ContentBrowserAssetDataSource/Private/NewAssetContextMenu.cpp
@@ -82,7 +82,10 @@ struct FCategorySubMenuItem
 	{
 		for (const FCategoryChildren& Child : Children)
 		{
-			OutSubMenuItems.Add(Child.SubMenuItem);
+			if (Child.SubMenuItem.IsValid())
+			{
+				OutSubMenuItems.Add(Child.SubMenuItem);
+			}
 		}
 	}
 
